PlainProps prototype

diff --git a/Engine/Source/Programs/PlainPropsTests/PlainPropsTests.Build.cs b/Engine/Source/Programs/PlainPropsTests/PlainPropsTests.Build.cs
new file mode 100644
--- /dev/null
+++ b/Engine/Source/Programs/PlainPropsTests/PlainPropsTests.Build.cs
@@ -0,0 +1,45 @@
+// Copyright Epic Games, Inc. All Rights Reserved.
+using UnrealBuildTool;
+
+public class PlainPropsTests : TestModuleRules
+{
+	protected Metadata PlainPropsMetadata = new Metadata() {
+		TestName = "PlainProps",
+		TestShortName = "PlainProps",
+		ReportType = "xml",
+		SupportedPlatforms = {
+			UnrealTargetPlatform.Win64,
+			UnrealTargetPlatform.Linux,
+			UnrealTargetPlatform.Mac} };
+
+	public Metadata TestMetadata
+	{ 
+		get { return PlainPropsMetadata; }
+	}
+
+	public PlainPropsTests(ReadOnlyTargetRules Target) : base(Target, true)
+	{
+		PrivateDependencyModuleNames.AddRange(
+			new string[] {
+				"Core",
+				"CoreUObject"
+			});
+
+		if (Target.bBuildWithEditorOnlyData)
+		{
+			PrivateDependencyModuleNames.AddRange(new string[] { "DesktopPlatform" });
+		}
+
+		string PlatformCompilationArgs;
+		foreach (var Platform in UnrealTargetPlatform.GetValidPlatforms())
+		{
+			PlatformCompilationArgs = "-allmodules";
+			TestMetadata.PlatformCompilationExtraArgs.Add(Platform, PlatformCompilationArgs);
+		}
+
+		// Platform-specific tags
+		TestMetadata.PlatformTags.Add(UnrealTargetPlatform.Linux, "~[.]~[Slow]");
+
+		UpdateBuildGraphPropertiesFile(TestMetadata);
+	}
+}
\ No newline at end of file
diff --git a/Engine/Source/Programs/PlainPropsTests/PlainPropsTests.Target.cs b/Engine/Source/Programs/PlainPropsTests/PlainPropsTests.Target.cs
new file mode 100644
--- /dev/null
+++ b/Engine/Source/Programs/PlainPropsTests/PlainPropsTests.Target.cs
@@ -0,0 +1,18 @@
+// Copyright Epic Games, Inc. All Rights Reserved.
+
+using UnrealBuildTool;
+
+[SupportedPlatforms(UnrealPlatformClass.All)]
+public class PlainPropsTestsTarget : TestTargetRules
+{
+	public PlainPropsTestsTarget(TargetInfo Target) : base(Target)
+	{
+		// Collects all tests decorated with #if WITH_LOW_LEVELTESTS from dependencies
+		bWithLowLevelTestsOverride = true;
+		bCompileWithPluginSupport = true;
+		bCompileAgainstCoreUObject = true;
+
+		bBuildWithEditorOnlyData = Target.Platform.IsInGroup(UnrealPlatformGroup.Desktop)
+			&& (Target.Configuration == UnrealTargetConfiguration.Debug || Target.Configuration == UnrealTargetConfiguration.Development);
+	}
+}
diff --git a/Engine/Source/Programs/PlainPropsTests/Private/PlainPropsBind.cpp b/Engine/Source/Programs/PlainPropsTests/Private/PlainPropsBind.cpp
new file mode 100644
--- /dev/null
+++ b/Engine/Source/Programs/PlainPropsTests/Private/PlainPropsBind.cpp
@@ -0,0 +1,215 @@
+// Copyright Epic Games, Inc. All Rights Reserved.
+
+#include  "PlainPropsBind.h"
+//#include  "PlainPropsInternalFormat.h"
+//#include  "PlainPropsInternalRead.h"
+//#include <type_traits>
+
+namespace PlainProps
+{
+
+static_assert(sizeof(ELeafBindType) == 1);
+static_assert((uint8)ELeafType::Bool		== (uint8)ELeafBindType::Bool);
+static_assert((uint8)ELeafType::IntS		== (uint8)ELeafBindType::IntS);
+static_assert((uint8)ELeafType::IntU		== (uint8)ELeafBindType::IntU);
+static_assert((uint8)ELeafType::Float		== (uint8)ELeafBindType::Float);
+static_assert((uint8)ELeafType::Hex			== (uint8)ELeafBindType::Hex);
+static_assert((uint8)ELeafType::Enum		== (uint8)ELeafBindType::Enum);
+static_assert((uint8)ELeafType::Unicode		== (uint8)ELeafBindType::Unicode);
+
+FMemberVisitor::FMemberVisitor(const FStructSchemaBinding& InSchema)
+: Schema(InSchema)
+, NumMembers(InSchema.NumMembers)
+{}
+
+EMemberKind FMemberVisitor::PeekKind() const
+{
+	return PeekType().GetKind();
+}
+
+FMemberBindType	FMemberVisitor::PeekType() const
+{
+	check(HasMore());
+	return Schema.Members[MemberIdx];
+}
+
+uint64 FMemberVisitor::GrabMemberOffset()
+{
+	return Schema.GetOffsets()[MemberIdx++];
+}
+
+FLeafMemberBinding FMemberVisitor::GrabLeaf()
+{
+	FUnpackedLeafBindType Leaf = PeekType().AsLeaf();
+	FOptionalEnumSchemaId Enum = Leaf.Type == ELeafBindType::Enum ? ToOptional(GrabEnumSchema()) : NoId;
+	uint64 Offset = GrabMemberOffset();
+
+	return {Leaf, Enum, Offset};
+}
+
+FStructMemberBinding FMemberVisitor::GrabStruct()
+{
+	checkf(!PeekType().AsStruct().IsDynamic, TEXT("Bound structs can't be dynamic"));
+	return { static_cast<FStructSchemaId>(GrabInnerSchema()), GrabMemberOffset() };
+}
+
+static bool HasSchema(FMemberBindType Type)
+{
+	checkf(!Type.IsStruct() || !Type.AsStruct().IsDynamic, TEXT("Bound structs can't be dynamic"));
+	return Type.IsStruct() || Type.AsLeaf().Bind.Type == ELeafBindType::Enum;
+}
+
+TConstArrayView<FMemberBindType> FMemberVisitor::GrabInnerTypes()
+{
+	const int32 Idx = InnerRangeIdx;
+	const TConstArrayView<FMemberBindType> All(Schema.GetInnerRangeTypes(), Schema.NumInnerRanges);
+	while (All[InnerRangeIdx++].IsRange()) {}
+	return All.Slice(Idx, InnerRangeIdx - Idx);
+}
+
+FRangeMemberBinding FMemberVisitor::GrabRange()
+{
+	ERangeSizeType MaxSize = PeekType().AsRange().MaxSize;
+	const FRangeBinding* RangeBindings = Schema.GetRangeBindings() + InnerRangeIdx;
+	check(MaxSize == RangeBindings[0].GetSizeType());
+	FMemberBindTypeRange InnerTypes = GrabInnerTypes();
+	FOptionalSchemaId InnermostSchema = HasSchema(InnerTypes.Last()) ? ToOptional(GrabInnerSchema()) : NoId; 
+	uint64 Offset = GrabMemberOffset();
+		
+	return { &InnerTypes[0], RangeBindings, InnermostSchema, Offset};
+}
+
+FSchemaId FMemberVisitor::GrabInnerSchema()
+{
+	check(InnerSchemaIdx < Schema.NumInnerSchemas);
+	return Schema.GetInnerSchemas()[InnerSchemaIdx++];
+}
+
+FRangeBinding::FRangeBinding(const IRangeBinding& Binding, ERangeSizeType SizeType)
+: Handle(uint64(&Binding) | uint8(SizeType))
+{
+	check(&Binding == &GetBinding());
+	check(SizeType == GetSizeType());
+}
+
+
+////////////////////////////////////////////////////////////////////////////////////////////////
+
+FStructBindings::~FStructBindings()
+{
+	for (FStructBinding Binding : Bindings)
+	{
+		if (FStructSchemaBinding* Schema = Binding.TryGetSchema())
+		{
+			FMemory::Free(Schema);
+		}
+	}
+}
+
+void FStructBindings::BindStruct(FStructSchemaId Id, const ICustomStructBinding& Custom)
+{
+	Bind(Id, FStructBinding(Custom));
+}
+
+static uint16 CountInnerSchemas(TConstArrayView<FMemberBinding> Members)
+{
+	uint32 Out = 0;
+	for (const FMemberBinding& Member : Members)
+	{
+		Out += !!Member.InnermostSchema;
+	}
+	return IntCastChecked<uint16>(Out);
+}
+
+static uint16 CountRanges(TConstArrayView<FMemberBinding> Members)
+{
+	int32 Out = 0;
+	for (const FMemberBinding& Member : Members)
+	{
+		Out += Member.RangeBindings.Num();
+	}
+	return IntCastChecked<uint16>(Out);
+}
+
+void FStructBindings::BindStruct(FStructSchemaId Id, TConstArrayView<FMemberBinding> Members)
+{
+	FStructSchemaBinding Header = { IntCastChecked<uint16>(Members.Num()), CountInnerSchemas(Members), CountRanges(Members) };
+
+	// Calculate size
+	uint64 Size = sizeof(Header);
+	Size += (Members.Num() + Header.NumInnerRanges) * sizeof(FMemberBindType);
+	Size = Align(Size, sizeof(uint32));
+	Size += Members.Num() * sizeof(uint32);
+	Size = Align(Size, sizeof(FSchemaId));
+	Size += Header.NumInnerSchemas * sizeof(FSchemaId);
+	Size = Align(Size, sizeof(FRangeBinding));
+	Size += Header.NumInnerRanges * sizeof(FRangeBinding);
+
+	// Allocate and copy header
+	FStructSchemaBinding* Schema = new (FMemory::MallocZeroed(Size)) FStructSchemaBinding {Header};
+
+	// Copy footer
+	FMemberBindType* MemberIt		= Schema->Members;
+	FMemberBindType* RangeTypeIt	= const_cast<FMemberBindType*>(Schema->GetInnerRangeTypes());
+	uint32* OffsetIt				= const_cast<uint32*>(Schema->GetOffsets());
+	FSchemaId* InnerSchemaIt		= const_cast<FSchemaId*>(Schema->GetInnerSchemas());
+	FRangeBinding* RangeBindingIt	= const_cast<FRangeBinding*>(Schema->GetRangeBindings());
+	for (const FMemberBinding& Member : Members)
+	{
+		if (TConstArrayView<FRangeBinding> Ranges = Member.RangeBindings; Ranges.IsEmpty())
+		{
+			*MemberIt++ = Member.InnermostType;
+		}
+		else
+		{
+			*MemberIt++ = FMemberBindType(Ranges[0].GetSizeType());
+
+			for (FRangeBinding Range : Ranges.RightChop(1))
+			{
+				*RangeTypeIt++ = FMemberBindType(Range.GetSizeType());
+			}
+			*RangeTypeIt++ = Member.InnermostType;
+
+			FMemory::Memcpy(RangeBindingIt, Ranges.GetData(), Ranges.Num() * Ranges.GetTypeSize());
+			RangeBindingIt += Ranges.Num();
+		}
+
+		if (Member.InnermostSchema)
+		{
+			*InnerSchemaIt++ = Member.InnermostSchema.Get();
+		}
+
+		*OffsetIt++ = IntCastChecked<uint32>(Member.Offset);
+	}
+
+	// Validate copying
+	check(MemberIt == Schema->GetInnerRangeTypes());
+	check(RangeTypeIt == (const void*)Schema->GetOffsets());
+	check(OffsetIt == (const void*)Schema->GetInnerSchemas());
+	check(InnerSchemaIt == (const void*)Schema->GetRangeBindings());
+	check(Header.NumInnerRanges == RangeBindingIt - Schema->GetRangeBindings());
+
+	Bind(Id, FStructBinding(*Schema));
+}
+
+void FStructBindings::Bind(FStructSchemaId Id, FStructBinding Binding)
+{
+	if (Id.Idx >= static_cast<uint32>(Bindings.Num()))
+	{
+		Bindings.SetNumZeroed(Id.Idx + 1);
+	}
+	check(!Bindings[Id.Idx].IsBound());
+
+	Bindings[Id.Idx] = Binding;
+}
+
+FStructBinding FStructBindings::Get(FStructSchemaId Id) const
+{
+	FStructBinding Out = Bindings[Id.Idx];
+	check(Out.IsBound());
+	return Out;
+}
+
+//void FStructBindings::DropStruct(FStructSchemaId Id);
+
+} // namespace PlainProps
\ No newline at end of file
diff --git a/Engine/Source/Programs/PlainPropsTests/Private/PlainPropsBuild.cpp b/Engine/Source/Programs/PlainPropsTests/Private/PlainPropsBuild.cpp
new file mode 100644
--- /dev/null
+++ b/Engine/Source/Programs/PlainPropsTests/Private/PlainPropsBuild.cpp
@@ -0,0 +1,298 @@
+// Copyright Epic Games, Inc. All Rights Reserved.
+
+#include "PlainPropsBuild.h"
+#include  "PlainPropsIndex.h"
+#include  "PlainPropsInternalBuild.h"
+#include  "PlainPropsInternalFormat.h"
+#include "Serialization/VarInt.h"
+#include "Templates/TypeCompatibleBytes.h"
+#include "Hash/xxhash.h"
+
+namespace PlainProps
+{
+
+FMemberSchema MakeNestedRangeSchema(ERangeSizeType SizeType, const FMemberSchema& InnerRangeSchema)
+{
+	FMemberSchema Out = { FMemberType(SizeType), InnerRangeSchema.InnerSchema };
+	Out.InnerRangeTypes.Reserve(1 + InnerRangeSchema.InnerRangeTypes.Num());
+	Out.InnerRangeTypes.Add(InnerRangeSchema.Type);
+	Out.InnerRangeTypes.Append(InnerRangeSchema.InnerRangeTypes);
+	return Out;
+}
+
+//////////////////////////////////////////////////////////////////////////
+
+template<typename BuiltType>
+FBuiltRange* BuildStructuralRangeImpl(/* in-out */ TArrayView64<BuiltType> Values)
+{
+	check(Values.Num() > 0);
+	static_assert(alignof(FBuiltRange) >= alignof(BuiltType));
+	FBuiltRange* Out = new (FMemory::Malloc(sizeof(FBuiltRange) + Values.Num() * sizeof(BuiltType))) FBuiltRange;
+	Out->Num = Values.Num();
+	BuiltType* DataIt = reinterpret_cast<BuiltType*>(Out->Data);
+	for (BuiltType& Value : Values)
+	{
+		new (DataIt++) BuiltType(MoveTemp(Value));
+	}
+
+	return Out;
+}
+
+namespace Private
+{
+	FBuiltRange* BuildStructuralRange(TArrayView64<FBuiltRange*> Values)
+	{
+		return PlainProps::BuildStructuralRangeImpl(Values);
+	}
+
+	FBuiltRange* BuildStructuralRange(TArrayView64<TUniquePtr<FBuiltStruct>> Values)
+	{
+		return PlainProps::BuildStructuralRangeImpl(Values);
+	}
+
+	FBuiltRange* BuildLeafRange(FUnpackedLeafType Leaf, uint64 Num, FMemoryView Values)
+	{
+		check(Num > 0);
+		check(Values.GetSize() == Num * SizeOf(Leaf.Width));
+
+		FBuiltRange* Out = new (FMemory::Malloc(sizeof(FBuiltRange) + Values.GetSize())) FBuiltRange;
+		Out->Num = Num;
+		FMemory::Memcpy(Out->Data, Values.GetData(), Values.GetSize());
+		return Out;
+	}
+	
+	template<typename FloatType>
+	void NormalizeFloats(FloatType* Values, uint64 Num)
+	{
+		for (uint64 Idx = 0; Idx < Num; ++Idx)
+		{
+			// Reject NaN / INF and ignore negative zero for now
+			checkf(FMath::IsFinite(Values[Idx]), TEXT("Saving NaN or INF isn't supported"));
+		}
+	}
+
+	void NormalizeLeafRange(FUnpackedLeafType Leaf, FBuiltRange& Out)
+	{
+		check(Leaf.Type == ELeafType::Float);
+		if (Leaf.Width == ELeafWidth::B32)
+		{
+			NormalizeFloats(reinterpret_cast<float*>(Out.Data), Out.Num);
+		}
+		else
+		{
+			check(Leaf.Width == ELeafWidth::B64);
+			NormalizeFloats(reinterpret_cast<double*>(Out.Data), Out.Num);
+		}
+	}
+} // namespace Private
+
+//////////////////////////////////////////////////////////////////////////
+
+void FMemberBuilder::AddStruct(FMemberId Name, FStructSchemaId Schema, TUniquePtr<FBuiltStruct>&& Struct)
+{
+	Members.Emplace(Name, Schema, MoveTemp(Struct));
+}
+
+void FMemberBuilder::AddLeaf(FMemberId Name, FUnpackedLeafType Leaf, FOptionalEnumSchemaId Enum, uint64 Value)
+{
+	Members.Emplace(Name, Leaf, Enum, Value);
+}
+
+void FMemberBuilder::AddRange(FMemberId Name, FTypedRange&& Range)
+{
+	Members.Emplace(Name, MoveTemp(Range));
+}
+
+TSet<FBuiltStruct*> GLiveStructsFoo;
+
+void FMemberBuilder::BuildSuperStruct(const FStructDeclaration& Super, const FDebugIds& Debug)
+{
+	// If we need to support EMemberPresence::RequireAll for sub structs,
+	// we need access to struct declaration here or create an empty super
+	// structs that we throw away in BuildAndReset.
+	if (Members.IsEmpty() || (Members.Num() == 1 && IsSuper(Members[0].Schema.Type)))
+	{
+		return;
+	}
+	
+	TUniquePtr<FBuiltStruct> OnlyMember = BuildAndReset(Super, Debug);
+	Members.Emplace(FBuiltMember::MakeSuper(Super.Id, MoveTemp(OnlyMember)));
+	check(Members[0].Schema.Type.AsStruct().IsSuper);
+}
+
+TUniquePtr<FBuiltStruct> FMemberBuilder::BuildAndReset(const FStructDeclaration& Declared, const FDebugIds& Debug)
+{
+	checkf(!(Declared.Super && Declared.Occupancy == EMemberPresence::RequireAll),
+		TEXT("Requiring sub structs to be dense isn't implemented"));
+#if DO_CHECK
+	int32 OrderIdx = 0;
+	TConstArrayView<FMemberId> Order = Declared.GetMemberOrder();
+	for (FBuiltMember& Member : Members)
+	{
+		if (Declared.Super && IsSuper(Member.Schema.Type) && &Member == &Members[0])
+		{
+			continue; // Skip undeclared generated super struct
+		}
+
+		for (; OrderIdx < Order.Num() && Order[OrderIdx] != Member.Name; ++OrderIdx)
+		{}	
+		checkf(OrderIdx < Order.Num(), TEXT("Member '%s' in '%s' %s"), *Debug.Print(Member.Name), *Debug.Print(Declared.Type),
+				Order.Contains(Member.Name) ? TEXT("appeared in non-declared order") : TEXT("is undeclared"));
+		++OrderIdx;
+	}
+#endif
+
+	uint32 Num = static_cast<uint32>(Members.Num());
+	SIZE_T NumBytes = sizeof(FBuiltStruct) + Num * sizeof(FBuiltMember);
+	FBuiltStruct* Out = reinterpret_cast<FBuiltStruct*>(FMemory::MallocZeroed(NumBytes, alignof(FBuiltStruct)));
+	Out->NumMembers = IntCastChecked<uint16>(Num);
+	for (FBuiltMember& Member : Members)
+	{
+		Out->Members[&Member - &Members[0]] = MoveTemp(Member);
+	}
+
+	Members.Reset();
+
+	GLiveStructsFoo.Add(Out);
+
+	return TUniquePtr<FBuiltStruct>(Out);
+}
+
+template<typename IntType, typename FloatType>
+inline IntType CheckFiniteBitCast(FloatType Value)
+{
+	// Reject NaN / INF and ignore negative zero for now
+	checkf(FMath::IsFinite(Value), TEXT("Saving NaN or INF isn't supported"));
+	return BitCast<IntType>(Value);
+}
+
+uint64 ValueCast(float Value)
+{
+	return CheckFiniteBitCast<uint32>(Value);
+}
+
+uint64 ValueCast(double Value)
+{ 
+	return CheckFiniteBitCast<uint64>(Value);
+}
+
+
+//////////////////////////////////////////////////////////////////////////
+
+FBuiltMember::FBuiltMember(FBuiltMember&& O)
+: Name(O.Name)
+, Schema(MoveTemp(O.Schema))
+, Value(O.Value)
+{
+	O.Name = NoId;
+	O.Schema.Type = FMemberType(ELeafType::Bool, ELeafWidth::B8);
+	O.Value.Leaf = 0;
+}
+
+FBuiltMember::FBuiltMember(FMemberId Name, FUnpackedLeafType Leaf, FOptionalEnumSchemaId Enum, uint64 Value)
+: FBuiltMember(Name, {Leaf.Pack(), FOptionalSchemaId(Enum)}, { .Leaf = Value})
+{}
+
+//FBuiltMember::FBuiltMember(FMemberId Name, FEnumSchemaId Schema, ELeafWidth Width, uint64 Value)
+//: FBuiltMember(Name, MakeMemberSchema(Width, Schema), { .Leaf = Value })
+//{}
+
+FBuiltMember::FBuiltMember(FMemberId Name, FTypedRange&& Range)
+: FBuiltMember(Name, MoveTemp(Range.Schema), { .Range = Range.Values })
+{}
+
+FBuiltMember::FBuiltMember(FMemberId Name, FStructSchemaId Schema, TUniquePtr<FBuiltStruct>&& Value)
+: FBuiltMember(Name, {DefaultStructType, FOptionalSchemaId(Schema)}, { .Struct = Value.Release() })
+{}
+
+FBuiltMember FBuiltMember::MakeSuper(FStructSchemaId Schema, TUniquePtr<FBuiltStruct>&& Value)
+{
+	return FBuiltMember(NoId, {SuperStructType, FOptionalSchemaId(Schema)}, { .Struct = Value.Release() });
+}
+
+FBuiltMember::~FBuiltMember()
+{
+	if (Schema.Type.GetKind() == EMemberKind::Struct)
+	{
+		delete Value.Struct;
+	}
+	else if (Schema.Type.GetKind() == EMemberKind::Range)
+	{
+		FBuiltRange::Delete(Value.Range, Schema.InnerSchema, Schema.InnerRangeTypes);
+	}
+}
+
+FBuiltMember& FBuiltMember::operator=(FBuiltMember&& O)
+{
+	this->~FBuiltMember();
+	new (this) FBuiltMember(MoveTemp(O));
+	return *this;
+}
+
+//////////////////////////////////////////////////////////////////////////
+
+uint64 FBuiltRange::Delete(FBuiltRange* Range, FOptionalSchemaId InnerSchema, TConstArrayView<FMemberType> InnerTypes)
+{ 
+	// TODO: Handle struct and nested ranges
+
+	switch (InnerTypes[0].GetKind()) 
+	{
+	case EMemberKind::Struct:
+		//DeleteStructs(InnerTypes[0].AsStruct(), static_cast<FStructSchemaId>(Schema.InnerSchema.Get()), Range);
+		break;
+	case EMemberKind::Range:
+		//for (FNestedRangeIterator It(Range); It;)
+		//{
+		//	FBuiltRange::Delete(&*It, InnerSchema, InnerTypes.RightChop(1));
+		//	It.Advance();
+		//}
+	case EMemberKind::Leaf:		break;
+	}
+
+	FMemory::Free(Range);
+
+	return 0; // todo: return number of bytes if needed?
+}
+
+//////////////////////////////////////////////////////////////////////////
+
+FBuiltStruct::~FBuiltStruct()
+{ 
+	check(GLiveStructsFoo.Remove(this) == 1);
+
+	for (const FBuiltMember& Member : MakeArrayView(Members, NumMembers))
+	{
+		Member.~FBuiltMember();
+	}
+}
+
+//////////////////////////////////////////////////////////////////////////
+
+FTypedRange FStructRangeBuilder::BuildAndReset(const FStructDeclaration& Declared, const FDebugIds& Debug)
+{
+	TArray64<TUniquePtr<FBuiltStruct>> BuiltStructs;
+	BuiltStructs.Reserve(Structs.Num());
+	for (FMemberBuilder& Struct : Structs)
+	{
+		BuiltStructs.Emplace(Struct.BuildAndReset(Declared, Debug));
+	}
+
+	return BuildStructRange(Declared.Id, SizeType, /* move */ BuiltStructs);
+}
+
+//////////////////////////////////////////////////////////////////////////
+
+FNestedRangeBuilder::~FNestedRangeBuilder()
+{
+	checkf(Ranges.IsEmpty(), TEXT("Half-built range, forgot to call BuildAndReset() before destruction?"));
+}
+
+FTypedRange FNestedRangeBuilder::BuildAndReset(ERangeSizeType SizeType)
+{
+	FTypedRange Out = { MakeNestedRangeSchema(SizeType, Schema), 
+						Ranges.IsEmpty() ? nullptr : Private::BuildStructuralRange(Ranges) };
+	Ranges.Reset();
+	return Out;
+}
+
+} // namespace PlainProps
\ No newline at end of file
diff --git a/Engine/Source/Programs/PlainPropsTests/Private/PlainPropsBuildSchema.cpp b/Engine/Source/Programs/PlainPropsTests/Private/PlainPropsBuildSchema.cpp
new file mode 100644
--- /dev/null
+++ b/Engine/Source/Programs/PlainPropsTests/Private/PlainPropsBuildSchema.cpp
@@ -0,0 +1,327 @@
+// Copyright Epic Games, Inc. All Rights Reserved.
+
+#include  "PlainPropsBuildSchema.h"
+#include  "PlainPropsInternalBuild.h"
+#include  "PlainPropsInternalFormat.h"
+#include "Algo/Compare.h"
+#include "Algo/Find.h"
+
+namespace PlainProps
+{
+
+static const TCHAR* ToString(FLeafType Leaf)
+{
+	const TCHAR* Leaves[8][4] = {
+		{TEXT("bool"),		TEXT("ERR_b16"),	TEXT("ERR_b32"),	TEXT("ERR_b64")},
+		{TEXT("int8"),		TEXT("int16"),		TEXT("int32"),		TEXT("int64")},
+		{TEXT("uint8"),		TEXT("uint16"),		TEXT("uint32"),		TEXT("uint64")},
+		{TEXT("ERR_fp8"),	TEXT("ERR_fp16"),	TEXT("float"),		TEXT("double")},
+		{TEXT("hex8"),		TEXT("hex16"),		TEXT("hex32"),		TEXT("hex64")},
+		{TEXT("enum8"),		TEXT("enum16"),		TEXT("enum32"),		TEXT("enum64")},
+		{TEXT("utf8"),		TEXT("utf16"),		TEXT("utf32"),		TEXT("ERR_utf64")},
+		{TEXT("ERR_oob"),	TEXT("ERR_oob"),	TEXT("ERR_oob"),	TEXT("ERR_oob")}};
+
+	return Leaves[(uint8)Leaf.Type][(uint8)Leaf.Width];
+}
+
+
+static FString PrintMemberSchema(FMemberType Type, FOptionalSchemaId InnerSchema, TConstArrayView<FMemberType> InnerRangeTypes)
+{
+	switch (Type.GetKind())
+	{
+	case EMemberKind::Leaf:		return ToString(Type.AsLeaf());
+	case EMemberKind::Struct:	return FString::Printf(TEXT("Struct [%d]%s%s"), InnerSchema.Get().Idx, Type.AsStruct().IsSuper ? TEXT(" (super)") : TEXT(""), Type.AsStruct().IsDynamic ? TEXT(" (dynamic)") : TEXT(""));
+	case EMemberKind::Range:	return TEXT("Range of ") + PrintMemberSchema(InnerRangeTypes[0], InnerSchema, InnerRangeTypes.RightChop(1));
+	}
+	return "Illegal member kind";
+}
+
+static FString PrintMemberSchema(FMemberSchema Schema)
+{
+	return PrintMemberSchema(Schema.Type, Schema.InnerSchema, Schema.InnerRangeTypes);
+}
+
+//////////////////////////////////////////////////////////////////////////
+
+static void MarkInheritanceChainUsed(TArrayView<FBuiltStructSchema> OutStructs)
+{
+	for (const FBuiltStructSchema& Struct : OutStructs)
+	{
+		if (Struct.bUsed)
+		{
+			for (FOptionalStructSchemaId Super = Struct.Super; Super && !OutStructs[Super.Get().Idx].bUsed; Super = OutStructs[Super.Get().Idx].Super)
+			{
+				OutStructs[Super.Get().Idx].bUsed = true;
+			}
+		}
+	}
+}
+
+FSchemasBuilder::FSchemasBuilder(TConstArrayView<TUniquePtr<FStructDeclaration>> InStructs, TConstArrayView<TUniquePtr<FEnumDeclaration>> InEnums, const FDebugIds& InDebug)
+: Debug(InDebug)
+{
+	Structs.Reserve(InStructs.Num());
+	Enums.Reserve(InEnums.Num());
+	for (const TUniquePtr<FStructDeclaration>& Declaration : InStructs)
+	{
+		Structs.Emplace(*Declaration, *this, Debug);
+	}
+	for (const TUniquePtr<FEnumDeclaration>& Declaration : InEnums)
+	{
+		Enums.Emplace(*Declaration);
+	}
+}
+
+FBuiltSchemas FSchemasBuilder::Build() const
+{
+	FBuiltSchemas Out;
+	Out.Structs.Reserve(Structs.Num());
+	for (const FStructSchemaBuilder& Struct : Structs)
+	{
+		Out.Structs.Add(Struct.Build());
+	}
+	Out.Enums.Reserve(Enums.Num());
+	for (const FEnumSchemaBuilder& Enum : Enums)
+	{
+		Out.Enums.Add(Enum.Build());
+	}
+
+	MarkInheritanceChainUsed(Out.Structs);
+
+	return Out;
+}
+
+//////////////////////////////////////////////////////////////////////////
+
+static bool RequiresDynamicStructSchema(const FMemberSchema& A, const FMemberSchema& B)
+{
+	if (A.InnerSchema != B.InnerSchema && A.Type.GetKind() == B.Type.GetKind())
+	{
+		if (A.Type.IsStruct())
+		{
+			return true;
+		}
+		else if (A.Type.IsRange() && A.InnerRangeTypes.Last().IsStruct() && B.InnerRangeTypes.Last().IsStruct())
+		{
+			// Same range size and nested range sizes
+			return	A.Type == B.Type &&	Algo::Compare(	MakeArrayView(A.InnerRangeTypes).LeftChop(1),
+														MakeArrayView(B.InnerRangeTypes).LeftChop(1));
+		}
+	}
+
+	return false;
+		
+}
+
+static void SetIsDynamic(FMemberType& InOut)
+{
+	FStructType Type = InOut.AsStruct();
+	Type.IsDynamic = true;
+	InOut = FMemberType(Type);
+}
+
+void FStructSchemaBuilder::NoteMembersRecursively(const FBuiltStruct& Struct)
+{
+	checkf(!bBuilt, TEXT("Noted new members after building, built schema lack these memebers! Build() also returns pointers into NotedMembers which mustn't grow."));
+	check(Declaration.Occupancy != EMemberPresence::RequireAll || Struct.NumMembers == Declaration.NumMembers);
+	bMissingMemberNoted |= Struct.NumMembers < static_cast<uint32>(Declaration.NumMembers + !!Declaration.Super);
+	
+	for (const FBuiltMember& Member : MakeArrayView(Struct.Members, Struct.NumMembers))
+	{
+		if (FMemberSchema* Schema = NotedMembers.Find(Member.Name))
+		{
+			if (RequiresDynamicStructSchema(*Schema, Member.Schema))
+			{
+				SetIsDynamic(Schema->Type.IsStruct() ? Schema->Type : Schema->InnerRangeTypes.Last());
+				Schema->InnerSchema = {};
+			}
+			else
+			{
+				checkf(*Schema == Member.Schema, TEXT("Member '%s' in '%s' first added as %s and later as %s."),
+					*Debug.Print(Member.Name), *Debug.Print(Declaration.Type), *PrintMemberSchema(*Schema), *PrintMemberSchema(Member.Schema));
+			}
+		}
+		else
+		{
+			NotedMembers.Add(Member.Name, Member.Schema);
+		}
+		
+		const FMemberSchema& Schema = Member.Schema;
+		if (EMemberKind Kind = Schema.Type.GetKind(); Kind == EMemberKind::Leaf)
+		{
+			if (ELeafType::Enum == Schema.Type.AsLeaf().Type)
+			{
+				AllSchemas.NoteValue(static_cast<FEnumSchemaId>(Schema.InnerSchema.Get()), Member.Value.Leaf);
+			}
+		}
+		else if (Kind == EMemberKind::Struct)
+		{
+			AllSchemas.NoteMembers(static_cast<FStructSchemaId>(Schema.InnerSchema.Get()), *Member.Value.Struct);
+		}
+		else if (Member.Value.Range && IsStructOrEnum(MakeArrayView(Schema.InnerRangeTypes).Last()))
+		{
+			NoteRangeRecursively(Schema.Type.AsRange().MaxSize, Schema.InnerRangeTypes, Schema.InnerSchema.Get(), *Member.Value.Range);
+		}
+	}
+}
+
+template<class T>
+TConstArrayView64<T> MakeArrayView(const T* Data, uint64 Num)
+{
+	return TConstArrayView64<T>(Data, static_cast<int64>(Num));
+}
+
+template<typename IntType>
+void NoteEnumValues(FSchemasBuilder& AllSchemas, FEnumSchemaId Enum, const IntType* Values, uint64 Num)
+{
+	for (IntType Value : TConstArrayView64<IntType>(Values, Num))
+	{
+		AllSchemas.NoteValue(Enum, Value);
+	}
+}
+
+static void NoteEnumRange(FSchemasBuilder& Out, FLeafType Leaf, FEnumSchemaId Enum, const FBuiltRange& Range)
+{
+	check(Leaf.Type == ELeafType::Enum);
+	switch (Leaf.Width)
+	{
+	case ELeafWidth::B8:	NoteEnumValues(Out, Enum, reinterpret_cast<const uint8* >(Range.Data), Range.Num); break;
+	case ELeafWidth::B16:	NoteEnumValues(Out, Enum, reinterpret_cast<const uint16*>(Range.Data), Range.Num); break;
+	case ELeafWidth::B32:	NoteEnumValues(Out, Enum, reinterpret_cast<const uint32*>(Range.Data), Range.Num); break;
+	case ELeafWidth::B64:	NoteEnumValues(Out, Enum, reinterpret_cast<const uint64*>(Range.Data), Range.Num); break;
+	}
+}
+
+void FStructSchemaBuilder::NoteRangeRecursively(ERangeSizeType NumType, TConstArrayView<FMemberType> Types, FSchemaId InnermostSchema, const FBuiltRange& Range)
+{
+	FMemberType Type = Types[0];
+	switch (Type.GetKind())
+	{
+	case EMemberKind::Struct:
+		for (const TUniquePtr<const FBuiltStruct>& Struct : Range.AsStructs())
+		{
+			AllSchemas.NoteMembers(static_cast<FStructSchemaId>(InnermostSchema), *Struct);
+		}
+		break;
+	case EMemberKind::Range:
+		for (const FBuiltRange* InnerRange : Range.AsRanges())
+		{
+			if (InnerRange)
+			{
+				NoteRangeRecursively(Type.AsRange().MaxSize, Types.RightChop(1), InnermostSchema, *InnerRange);
+			}
+		}
+		break;
+	case EMemberKind::Leaf:
+		NoteEnumRange(/* out */ AllSchemas, Type.AsLeaf(), static_cast<FEnumSchemaId>(InnermostSchema), Range);
+		break;
+	}
+}
+
+FBuiltStructSchema FStructSchemaBuilder::Build() const
+{
+	FBuiltStructSchema Out = { Declaration.Type };
+	Out.bDense = Declaration.Occupancy == EMemberPresence::RequireAll || !bMissingMemberNoted;
+	Out.Super = Declaration.Super;
+	
+	if (int32 Num = NotedMembers.Num())
+	{
+		Out.bUsed = true;
+		Out.MemberNames.Reserve(Num);
+		Out.MemberSchemas.Reserve(Num);
+
+		// Add generated super struct before declared members
+		if (Declaration.Super)
+		{
+			for (const TPair<FOptionalMemberId, FMemberSchema>& NotedMember : NotedMembers)
+			{
+				if (NotedMember.Key == NoId)
+				{
+					Out.MemberSchemas.Add(&NotedMember.Value);
+					break;
+				}
+			}
+		}
+
+		for (FMemberId Name : Declaration.GetMemberOrder())
+		{
+			if (const FMemberSchema* Schema = NotedMembers.Find(ToOptional(Name)))
+			{
+				Out.MemberNames.Add(Name);
+				Out.MemberSchemas.Add(Schema);
+			}
+		}
+	}
+
+	bBuilt = true; 
+	check(NotedMembers.Num() == Out.MemberSchemas.Num());
+	return Out;
+}
+
+//////////////////////////////////////////////////////////////////////////
+
+FBuiltEnumSchema FEnumSchemaBuilder::Build() const
+{
+	FBuiltEnumSchema Out = { Declaration.Type };
+	Out.Mode = Declaration.Mode;
+	Out.Width = Declaration.Width;
+
+	if (int32 Num = NotedConstants.Num())
+	{
+		Out.bUsed = true;
+		Out.Names.Reserve(Num);
+		Out.Constants.Reserve(Num);
+		for (FEnumerator Enumerator : Declaration.GetEnumerators())
+		{
+			if (NotedConstants.Contains(Enumerator.Constant))
+			{
+				Out.Names.Add(Enumerator.Name);
+				Out.Constants.Add(Enumerator.Constant);
+			}
+		}
+	}
+	
+	bBuilt = true; 
+	check(	NotedConstants.Num() == Out.Constants.Num() || 
+			NotedConstants.Num() == Out.Constants.Num() + (Out.Mode == EEnumMode::Flag) );
+	return Out;
+}
+
+void FEnumSchemaBuilder::NoteValue(uint64 Value)
+{
+	if (Declaration.Mode == EEnumMode::Flag)
+	{
+		if (Value == 0)
+		{
+			// Don't validate 0 flag is declared, it isn't
+			NotedConstants.Add(Value);
+		}
+		else
+		{
+			const int32 NumValidated = NotedConstants.Num();
+			while (Value != 0)
+			{
+				uint64 HiBit = uint64(1) << FMath::FloorLog2_64(Value);
+				NotedConstants.Add(HiBit);
+				Value &= ~HiBit;
+			}
+
+			for (int32 Idx = NumValidated, Num = NotedConstants.Num(); Idx < Num; ++Idx)
+			{
+				uint64 Flag = NotedConstants.Get(FSetElementId::FromInteger(Idx));
+				checkf(Algo::FindBy(Declaration.GetEnumerators(), Flag, &FEnumerator::Constant), TEXT("Enum flag %d is undeclared"), Flag);
+			}
+		}
+	}
+	else
+	{
+		bool bValidated;
+		NotedConstants.FindOrAdd(Value, /* out */ &bValidated);
+		if (!bValidated)
+		{
+			checkf(Algo::FindBy(Declaration.GetEnumerators(), Value, &FEnumerator::Constant), TEXT("Enum value %d is undeclared"), Value);
+		}
+	}
+}
+
+} // namespace PlainProps
diff --git a/Engine/Source/Programs/PlainPropsTests/Private/PlainPropsDeclare.cpp b/Engine/Source/Programs/PlainPropsTests/Private/PlainPropsDeclare.cpp
new file mode 100644
--- /dev/null
+++ b/Engine/Source/Programs/PlainPropsTests/Private/PlainPropsDeclare.cpp
@@ -0,0 +1,80 @@
+// Copyright Epic Games, Inc. All Rights Reserved.
+
+#include  "PlainPropsDeclare.h"
+#include "Containers/Set.h"
+
+namespace PlainProps
+{
+
+static void ValidateDeclaration(const FEnumDeclaration& Enum)
+{
+	if (Enum.Mode == EEnumMode::Flag)
+	{
+		for (FEnumerator E : Enum.GetEnumerators())
+		{
+			checkf(FMath::CountBits(E.Constant) == 1, TEXT("Flag enums must use one bit per enumerator"));
+		}
+	}
+
+	TSet<uint32, DefaultKeyFuncs<uint32>, TInlineSetAllocator<64>> Names;
+	TSet<uint64, DefaultKeyFuncs<uint64>,TInlineSetAllocator<64>> Constants;
+	uint64 LastConstant = 0;
+	for (FEnumerator E : Enum.GetEnumerators())
+	{
+		checkf(FMath::FloorLog2_64(E.Constant) < 8 * SizeOf(Enum.Width), TEXT("Enumerator constant larger than declared width"));
+
+		bool bDeclared;
+		Names.FindOrAdd(E.Name.Idx, /* out*/ &bDeclared);
+		checkf(!bDeclared, TEXT("Enumerator name declared twice"));
+		Constants.FindOrAdd(E.Constant, /* out*/ &bDeclared);
+		checkf(!bDeclared, TEXT("Enumerator constant declared twice"));
+
+		checkf(LastConstant <= E.Constant, TEXT("Enumerator constants must be declared in ascending order"));
+		LastConstant = E.Constant;
+	}
+}
+
+template<typename T>
+void CopyItems(T* It, TConstArrayView<T> Items)
+{
+	for (T Item : Items)
+	{
+		(*It++)  = Item;
+	}
+}
+
+void FDeclarations::DeclareStruct(FStructSchemaId Id, FTypeId Type, TConstArrayView<FMemberId> MemberOrder, EMemberPresence Occupancy, FOptionalStructSchemaId Super)
+{
+	if (static_cast<int32>(Id.Idx) >= DeclaredStructs.Num())
+	{
+		DeclaredStructs.SetNum(Id.Idx + 1);
+	}
+
+	TUniquePtr<FStructDeclaration>& Ptr = DeclaredStructs[Id.Idx];
+	checkf(!Ptr, TEXT("Already declared"));
+
+	FStructDeclaration Header{Id, Type, Super, Occupancy, IntCastChecked<uint16>(MemberOrder.Num())};
+	void* Data = FMemory::Malloc(sizeof(FStructDeclaration) + MemberOrder.Num() * MemberOrder.GetTypeSize());
+	Ptr.Reset(new (Data) FStructDeclaration(Header));
+	CopyItems(Ptr->MemberOrder, MemberOrder);
+}
+
+void FDeclarations::DeclareEnum(FEnumSchemaId Id, FTypeId Type, EEnumMode Mode, ELeafWidth Width, TConstArrayView<FEnumerator> Enumerators)
+{
+	if (static_cast<int32>(Id.Idx) >= DeclaredEnums.Num())
+	{
+		DeclaredEnums.SetNum(Id.Idx + 1);
+	}
+	
+	TUniquePtr<FEnumDeclaration>& Ptr = DeclaredEnums[Id.Idx];
+	checkf(!Ptr, TEXT("Already declared")); 
+
+	FEnumDeclaration Header{Type, Mode, Width, IntCastChecked<uint16>(Enumerators.Num())};
+	void* Data = FMemory::Malloc(sizeof(FEnumDeclaration) + Enumerators.Num() * Enumerators.GetTypeSize());
+	Ptr.Reset(new (Data) FEnumDeclaration(Header));
+	CopyItems(Ptr->Enumerators, Enumerators);
+
+	ValidateDeclaration(*Ptr);
+}
+
+} // namespace PlainProps
\ No newline at end of file
diff --git a/Engine/Source/Programs/PlainPropsTests/Private/PlainPropsIndex.cpp b/Engine/Source/Programs/PlainPropsTests/Private/PlainPropsIndex.cpp
new file mode 100644
--- /dev/null
+++ b/Engine/Source/Programs/PlainPropsTests/Private/PlainPropsIndex.cpp
@@ -0,0 +1,156 @@
+// Copyright Epic Games, Inc. All Rights Reserved.
+
+#include  "PlainPropsIndex.h"
+#include "Hash/xxhash.h"
+
+namespace PlainProps
+{
+
+FNestedScopeId FNestedScopeIndexer::Index(FNestedScope Scope)
+{
+	uint32 Hash = GetTypeHash(Scope);
+	FSetElementId Idx = Scopes.FindIdByHash(Hash, Scope);
+	if (Idx == FSetElementId())
+	{
+		Idx = Scopes.AddByHash(Hash, Scope);
+	}
+	return { static_cast<uint32>(Idx.AsInteger()) };
+}
+
+FNestedScope FNestedScopeIndexer::Resolve(FNestedScopeId Id) const
+{
+	return Scopes[FSetElementId::FromInteger(Id.Idx)];
+}
+
+//////////////////////////////////////////////////////////////////////////
+
+FParametricTypeView::FParametricTypeView(FConcreteTypenameId InName, TConstArrayView<FTypeId> Params)
+: FParametricTypeView(InName, IntCastChecked<uint8>(Params.Num()), Params.GetData())
+{}
+
+namespace ParametricTypeHash
+{
+
+static uint32 Calculate(FConcreteTypenameId Name, TConstArrayView<FTypeId> Parameters)
+{
+	uint64 ParametersHash = FXxHash64::HashBuffer(Parameters.GetData(), sizeof(FTypeId) * Parameters.Num()).Hash;
+	return HashCombineFast(Name.Id.Idx, static_cast<uint32>(ParametersHash));
+}
+
+static constexpr uint8 FreeSlotByte = 0xFF;
+static constexpr uint32 FreeSlot = 0xFFFFFFFF;
+
+static uint32* Rehash(uint32 NumSlots, TConstArrayView<FParametricType> Types, TConstArrayView<FTypeId> Parameters)
+{
+	uint32* Slots = (uint32*)FMemory::Malloc(NumSlots * sizeof(uint32));
+	FMemory::Memset(Slots, FreeSlotByte, NumSlots * sizeof(uint32));
+
+	const uint32 SlotMask = NumSlots - 1;
+	uint32 TypeIdx = 0;
+	for (FParametricType Type : Types)
+	{
+		uint32 Hash = Calculate(Type.Name, Parameters.Slice(Type.Parameters.Idx, Type.Parameters.NumParameters));
+		uint32 SlotIdx = Hash & SlotMask;
+		while (Slots[SlotIdx] != FreeSlot)
+		{
+			SlotIdx = (SlotIdx + 1) & SlotMask;
+		}
+		Slots[SlotIdx] = TypeIdx;
+		++TypeIdx;
+	}
+
+	return Slots;
+}
+
+}
+
+FParametricTypeIndexer::~FParametricTypeIndexer()
+{
+	FMemory::Free(Slots);
+}
+
+FParametricTypeId FParametricTypeIndexer::Index(FParametricTypeView View)
+{
+	using namespace ParametricTypeHash;
+
+	static constexpr uint32 MinSlack = 4;
+	uint32 WantedSlots = FMath::RoundUpToPowerOfTwo((Types.Num() + MinSlack) * /* 90% load factor */ 10 / 9);
+	if (WantedSlots > NumSlots)
+	{
+		FMemory::Free(Slots);
+		Slots = Rehash(WantedSlots, Types, Parameters);
+		NumSlots = WantedSlots;
+	}
+
+	const uint32 Hash = Calculate(View.Name, MakeArrayView(View.Parameters, View.NumParameters));
+	uint32 SlotIdx = Hash & (NumSlots - 1);
+
+	for (TArrayView<uint32> SlotsPart : { MakeArrayView(Slots + SlotIdx, NumSlots - SlotIdx), MakeArrayView(Slots, SlotIdx) } )
+	{
+		for (uint32& Slot : SlotsPart)
+		{
+			if (Slot == FreeSlot)
+			{
+				Slot = Types.Num();
+				FParameterIndexRange ParameterIndices(View.NumParameters, Parameters.Num());
+				Types.Add({View.Name, ParameterIndices});
+				Parameters.Append(View.Parameters, View.NumParameters);
+				
+				return FParametricTypeId(View.NumParameters, Slot);
+			}
+			
+			FParametricType Existing = Types[Slot];
+			if (View.Name == Existing.Name &&
+				View.NumParameters == Existing.Parameters.NumParameters &&
+				!FMemory::Memcmp(View.Parameters, &Parameters[Existing.Parameters.Idx], View.NumParameters * sizeof(FTypeId)))
+			{
+				return FParametricTypeId(View.NumParameters, Slot);
+			}
+		}
+	}
+
+	checkf(false, TEXT("No free slots despite MinSlack and load factor"));
+	return FParametricTypeId(0,0);
+}
+
+FParametricTypeView FParametricTypeIndexer::Resolve(FParametricTypeId Id) const
+{
+	FParametricType Type = Types[Id.Idx];
+	check(Id.NumParameters == Type.Parameters.NumParameters);
+	return { Type.Name, static_cast<uint8>(Id.NumParameters), &Parameters[Type.Parameters.Idx] };
+}
+
+//////////////////////////////////////////////////////////////////////////
+
+FScopeId FIdIndexerBase::NestScope(FScopeId Outer, FFlatScopeId Inner)
+{
+	return FScopeId(NestedScopes.Index(Outer, Inner));
+}
+
+FTypeId FIdIndexerBase::MakeParametricType(FTypeId Type, TConstArrayView<FTypeId> Params)
+{
+	FParametricTypeId ParametricType = ParametricTypes.Index({Type.Name.AsConcrete(), IntCastChecked<uint8>(Params.Num()), Params.GetData()});
+	return {Type.Scope, FTypenameId(ParametricType)};
+}
+
+FEnumSchemaId FIdIndexerBase::IndexEnum(FTypeId Type)
+{
+	return { IntCastChecked<uint32>(Enums.Add(Type).AsInteger()) };
+}
+
+FStructSchemaId	FIdIndexerBase::IndexStruct(FTypeId Type)
+{
+	return { IntCastChecked<uint32>(Structs.Add(Type).AsInteger()) };
+}
+//
+//FParametricTypeView	FIdIndexerBase::Resolve(FParametricTypeId Id) const
+//{
+//	return ParametricTypes.Resolve(Id);
+//}
+//
+//FNestedScope FIdIndexerBase::Resolve(FNestedScopeId Id) const
+//{
+//	return NestedScopes.Resolve(Id);
+//}
+
+} // namespace PlainProps
\ No newline at end of file
diff --git a/Engine/Source/Programs/PlainPropsTests/Private/PlainPropsInternalBuild.h b/Engine/Source/Programs/PlainPropsTests/Private/PlainPropsInternalBuild.h
new file mode 100644
--- /dev/null
+++ b/Engine/Source/Programs/PlainPropsTests/Private/PlainPropsInternalBuild.h
@@ -0,0 +1,103 @@
+// Copyright Epic Games, Inc. All Rights Reserved.
+
+#pragma once
+
+#include  "PlainPropsBuild.h"
+#include "Containers/Array.h"
+#include "Containers/UnrealString.h"
+
+namespace PlainProps
+{
+
+union FBuiltValue
+{
+	uint64			Leaf;
+	FBuiltStruct*	Struct;
+	FBuiltRange*	Range;
+};
+
+struct FBuiltMember
+{
+	FBuiltMember() = delete;
+	FBuiltMember(const FBuiltMember&) = delete;
+	FBuiltMember(FBuiltMember&& O);
+	FBuiltMember(FMemberId Name, FUnpackedLeafType Leaf, FOptionalEnumSchemaId Schema, uint64 Value);
+	//FBuiltMember(FMemberId Name, FEnumSchemaId Schema, ELeafWidth Width,  uint64 Value);
+	FBuiltMember(FMemberId Name, FTypedRange&& Range);
+	FBuiltMember(FMemberId Name, FStructSchemaId Schema, TUniquePtr<FBuiltStruct>&& Value);
+	static FBuiltMember MakeSuper(FStructSchemaId Schema, TUniquePtr<FBuiltStruct>&& Value);
+	~FBuiltMember(); // Deletes Value
+
+	FBuiltMember& operator=(const FBuiltMember&) = delete;
+	FBuiltMember& operator=(FBuiltMember&&);
+
+	FOptionalMemberId		Name;
+	FMemberSchema			Schema;
+	FBuiltValue				Value;
+
+private:
+	FBuiltMember(FOptionalMemberId N, FMemberSchema&& S, FBuiltValue V) : Name(N), Schema(MoveTemp(S)), Value(V) {}
+};
+
+struct FBuiltStruct
+{
+	~FBuiltStruct();
+
+	uint16				NumMembers;
+	FBuiltMember		Members[0];
+};
+
+struct FBuiltRange
+{
+	~FBuiltRange() = delete;
+	static uint64 Delete(FBuiltRange* Range, FOptionalSchemaId InnerSchema, TConstArrayView<FMemberType> InnerTypes);
+
+	uint64				Num;
+	uint8				Data[0];
+	
+	TConstArrayView64<const FBuiltRange*> AsRanges() const { return { reinterpret_cast<FBuiltRange const* const*>(Data), IntCastChecked<int64>(Num) }; }
+	TConstArrayView64<TUniquePtr<const FBuiltStruct>> AsStructs() const { return { reinterpret_cast<const TUniquePtr<const FBuiltStruct>*>(Data), IntCastChecked<int64>(Num) }; }
+};
+
+//////////////////////////////////////////////////////////////////////////
+
+inline void WriteData(TArray64<uint8>& Out, const void* Data, int64 Size)
+{
+	Out.Append(static_cast<const uint8*>(Data), Size);
+}
+
+template<typename ArrayType>
+void WriteArray(TArray64<uint8>& Out, const ArrayType& In)
+{
+	static_assert(TIsContiguousContainer<ArrayType>::Value);
+	WriteData(Out, In.GetData(), sizeof(typename ArrayType::ElementType) * In.Num());
+}
+
+template<typename T>
+void WriteAlignmentPadding(TArray64<uint8>& Out)
+{
+	Out.AddZeroed(Align(Out.Num(), alignof(T)) - Out.Num());
+}	
+
+template<typename T>
+void WriteAlignedArray(TArray64<uint8>& Out, TArrayView<T> In)
+{
+	WriteAlignmentPadding<T>(Out);
+	WriteArray(Out, In);
+}
+
+template<typename T>
+inline void WriteInt(TArray64<uint8>& Out, T Number)
+{
+	static_assert(std::is_integral_v<T>);
+	static_assert(PLATFORM_LITTLE_ENDIAN);
+	WriteData(Out, &Number, sizeof(T));
+}
+
+inline void WriteU32(TArray64<uint8>& Out, uint32 Int) { WriteInt(Out, Int); }
+inline void WriteU64(TArray64<uint8>& Out, uint64 Int) { WriteInt(Out, Int); }
+
+PP_API uint64 WriteSkippableSlice(TArray64<uint8>& Out, TConstArrayView64<uint8> Slice);
+
+
+} // namespace PlainProps
\ No newline at end of file
diff --git a/Engine/Source/Programs/PlainPropsTests/Private/PlainPropsInternalFormat.h b/Engine/Source/Programs/PlainPropsTests/Private/PlainPropsInternalFormat.h
new file mode 100644
--- /dev/null
+++ b/Engine/Source/Programs/PlainPropsTests/Private/PlainPropsInternalFormat.h
@@ -0,0 +1,139 @@
+// Copyright Epic Games, Inc. All Rights Reserved.
+
+#pragma once
+
+#include  "PlainPropsTypes.h"
+
+namespace PlainProps
+{
+
+struct FSchemaBatch
+{
+	uint32 NumNestedScopes;
+	uint32 NestedScopesOffset;
+	uint32 NumParametricTypes;
+	uint32 NumSchemas;
+	uint32 SchemaOffsets[0];
+
+	TConstArrayView<FNestedScope> GetNestedScopes() const
+	{
+		return MakeArrayView(reinterpret_cast<const FNestedScope*>(reinterpret_cast<const uint8*>(this) + NestedScopesOffset), NumNestedScopes);
+	}
+
+	TConstArrayView<FParametricType> GetParametricTypes() const
+	{
+		return MakeArrayView(reinterpret_cast<const FParametricType*>(GetNestedScopes().GetData() + NumNestedScopes), NumParametricTypes);
+	}
+
+	const FTypeId* GetFirstParameter() const
+	{
+		return reinterpret_cast<const FTypeId*>(GetParametricTypes().GetData() + NumParametricTypes);
+	}
+
+	void ValidateBounds(uint64 NumBytes) const;
+};
+
+static_assert(sizeof(FSchemaBatch) == 16 && alignof(FSchemaBatch) == 4,		"Add binary format versioning and read old format in AllocateReadSchemas");
+static_assert(offsetof(FSchemaBatch, NumNestedScopes) == 0,						"Add binary format versioning and read old format in AllocateReadSchemas");
+static_assert(offsetof(FSchemaBatch, NestedScopesOffset) == 4,					"Add binary format versioning and read old format in AllocateReadSchemas");
+static_assert(offsetof(FSchemaBatch, NumParametricTypes) == 8,					"Add binary format versioning and read old format in AllocateReadSchemas");
+static_assert(offsetof(FSchemaBatch, NumSchemas) == 12,							"Add binary format versioning and read old format in AllocateReadSchemas");
+
+//////////////////////////////////////////////////////////////////////////
+
+enum class ESuper : uint8 { No, Unused, Used, Reused };
+
+inline bool UsesSuper(ESuper Inheritance) { return Inheritance == ESuper::Used || Inheritance == ESuper::Reused; }
+
+struct FStructSchema
+{
+	FTypeId Type;
+	uint16 NumMembers;
+	uint16 NumRangeTypes;
+	uint16 NumInnerSchemas;
+	ESuper Inheritance : 2;
+	uint8 IsDense : 1;
+	// Needed? maybe sufficient check sizes match up
+	//uint8 IsContiguous : 1; // Dense and only non-bool leaves. Could extend to handle nested contiguous structs by complicating format slightly.
+
+	FMemberType Footer[0];
+	
+	static const FMemberType* GetMemberTypes(const FMemberType* Footer)
+	{
+		return Footer;
+	}
+	
+	static const FMemberType* GetRangeTypes(const FMemberType* Footer, uint32 NumMembers)
+	{
+		return Footer + NumMembers;
+	}
+
+	static const FMemberId*	GetMemberNames(const FMemberType* Footer, uint32 NumMembers, uint32 NumRangeTypes)
+	{
+		return AlignPtr<FMemberId>(Footer + NumMembers + NumRangeTypes);
+	}
+	
+	static const FSchemaId*	GetInnerSchemas(const FMemberType* Footer, uint32 NumMembers, uint32 NumRangeTypes, uint32 NumNames)
+	{
+		return AlignPtr<FSchemaId>(GetMemberNames(Footer, NumMembers, NumRangeTypes) + NumNames);
+	}
+
+	const FSchemaId* GetInnerSchemas() const
+	{
+		return GetInnerSchemas(Footer, NumMembers, NumRangeTypes, NumMembers - UsesSuper(Inheritance));
+	}
+
+	FOptionalStructSchemaId GetSuperSchema() const
+	{
+		const FSchemaId* FirstSchema = GetInnerSchemas();
+		return Inheritance != ESuper::No ? ToOptional(static_cast<FStructSchemaId>(*FirstSchema)) : NoId;
+	}
+};
+
+static_assert(sizeof(FStructSchema) == 16 && alignof(FStructSchema) == 4,		"Add binary format versioning and read old format in AllocateReadSchemas");
+static_assert(offsetof(FStructSchema, Type) == 0,								"Add binary format versioning and read old format in AllocateReadSchemas");
+static_assert(offsetof(FTypeId, Scope) == 0,									"Add binary format versioning and read old format in AllocateReadSchemas");
+static_assert(offsetof(FTypeId, Name) == 4,										"Add binary format versioning and read old format in AllocateReadSchemas");
+static_assert(offsetof(FStructSchema, NumMembers) == 8,							"Add binary format versioning and read old format in AllocateReadSchemas");
+static_assert(offsetof(FStructSchema, NumRangeTypes) == 10,						"Add binary format versioning and read old format in AllocateReadSchemas");
+static_assert(offsetof(FStructSchema, NumInnerSchemas) == 12,					"Add binary format versioning and read old format in AllocateReadSchemas");
+static_assert(offsetof(FStructSchema, Footer) == 15,							"Add binary format versioning and read old format in AllocateReadSchemas");
+
+//////////////////////////////////////////////////////////////////////////
+
+struct FEnumSchema
+{
+	FTypeId Type;
+	uint8 FlagMode : 1;
+	uint8 ExplicitConstants : 1; // 0,1,2,3,4.. or 1,2,4,8..
+	ELeafWidth Width;
+	uint16 Num;
+	FNameId Footer[0];
+};
+
+static_assert(sizeof(FEnumSchema) == 12 && alignof(FEnumSchema) == 4,	"Add binary format versioning and read old format in AllocateReadSchemas");
+static_assert(offsetof(FEnumSchema, Type) == 0,							"Add binary format versioning and read old format in AllocateReadSchemas");
+static_assert(offsetof(FEnumSchema, Width) == 9,						"Add binary format versioning and read old format in AllocateReadSchemas");
+static_assert(offsetof(FEnumSchema, Num) == 10,							"Add binary format versioning and read old format in AllocateReadSchemas");
+static_assert(offsetof(FEnumSchema, Footer) == 12,						"Add binary format versioning and read old format in AllocateReadSchemas");
+
+//////////////////////////////////////////////////////////////////////////
+
+inline bool IsStructOrEnum(FMemberType Type)
+{
+	return Type.GetKind() == EMemberKind::Struct || (Type.GetKind() == EMemberKind::Leaf && ELeafType::Enum == Type.AsLeaf().Type);
+}
+
+inline bool IsSuper(FMemberType Type)
+{
+	return Type.IsStruct() && Type.AsStruct().IsSuper;
+}
+
+inline constexpr uint64 GetLeafRangeSize(uint64 Num, FUnpackedLeafType Leaf)
+{
+	return Leaf.Type == ELeafType::Bool ? (Num + 7) / 8 : Num * SizeOf(Leaf.Width);
+}
+
+//////////////////////////////////////////////////////////////////////////
+
+} // namespace PlainProps
\ No newline at end of file
diff --git a/Engine/Source/Programs/PlainPropsTests/Private/PlainPropsInternalRead.h b/Engine/Source/Programs/PlainPropsTests/Private/PlainPropsInternalRead.h
new file mode 100644
--- /dev/null
+++ b/Engine/Source/Programs/PlainPropsTests/Private/PlainPropsInternalRead.h
@@ -0,0 +1,40 @@
+// Copyright Epic Games, Inc. All Rights Reserved.
+
+#pragma once
+
+#include  "PlainPropsTypes.h"
+#include  "PlainPropsRead.h"
+
+namespace PlainProps
+{
+
+template<typename MemberType, typename IndexType>
+TConstArrayView<MemberType> GrabInnerRangeTypes(TConstArrayView<MemberType> InnerRangeTypes, IndexType& InOutIdx)
+{
+	const int32 StartIdx = InOutIdx;
+	int32 Idx = StartIdx;
+	while (InnerRangeTypes[Idx++].IsRange());
+	
+	InOutIdx = static_cast<IndexType>(Idx);
+	return MakeArrayView(&InnerRangeTypes[StartIdx], Idx - StartIdx);
+}
+
+inline uint64 GrabRangeNum(ERangeSizeType MaxSize, FByteReader& ByteIt, FBitCacheReader& BitIt)
+{
+	switch(MaxSize)
+	{
+	case ERangeSizeType::Uni:	return BitIt.GrabNext(/* in-out */ ByteIt) ? 1 : 0;
+	case ERangeSizeType::S8:	return IntCastChecked<uint64>(	ByteIt.Grab<int8>());
+	case ERangeSizeType::U8:	return							ByteIt.Grab<uint8>();
+	case ERangeSizeType::S16:	return IntCastChecked<uint64>(	ByteIt.Grab<int16>());
+	case ERangeSizeType::U16:	return							ByteIt.Grab<uint16>();
+	case ERangeSizeType::S32:	return IntCastChecked<uint64>(	ByteIt.Grab<int32>());
+	case ERangeSizeType::U32:	return							ByteIt.Grab<uint32>();
+	case ERangeSizeType::S64:	return IntCastChecked<uint64>(	ByteIt.Grab<int64>());
+	case ERangeSizeType::U64:	return							ByteIt.Grab<uint64>();
+	}
+	check(false);
+	return 0;
+}
+
+} // namespace PlainProps
\ No newline at end of file
diff --git a/Engine/Source/Programs/PlainPropsTests/Private/PlainPropsLoad.cpp b/Engine/Source/Programs/PlainPropsTests/Private/PlainPropsLoad.cpp
new file mode 100644
--- /dev/null
+++ b/Engine/Source/Programs/PlainPropsTests/Private/PlainPropsLoad.cpp
@@ -0,0 +1,502 @@
+// Copyright Epic Games, Inc. All Rights Reserved.
+
+#include  "PlainPropsLoad.h"
+#include  "PlainPropsBind.h"
+#include  "PlainPropsInternalFormat.h"
+#include  "PlainPropsInternalRead.h"
+#include <type_traits>
+
+namespace PlainProps
+{
+
+struct FLoadStructMemcpy
+{
+	uint32			Size;
+	uint32			Offset;
+};
+
+// Describes how to load a saved struct into the matching in-memory representation
+class FLoadStructPlan
+{
+public:
+	FLoadStructPlan() = default;
+
+	explicit FLoadStructPlan(FLoadStructMemcpy Memcpy)
+	: Handle((uint64(Memcpy.Size) << 32) | (uint64(Memcpy.Offset) << 2) | MemcpyMask)
+	{
+		check(Memcpy.Offset == AsMemcpy().Offset && Memcpy.Size == AsMemcpy().Size);
+	}
+
+	explicit FLoadStructPlan(const ICustomStructBinding& Custom)
+	: Handle(uint64(&Custom) | CustomMask)
+	{
+		check(&Custom == &AsCustom());
+	}
+
+	// @param OffsetWidth Usage unimplemented, store size and offsets as 8/16/32/64-bit 
+	explicit FLoadStructPlan(const FStructSchemaBinding& Schema, ELeafWidth OffsetWidth, bool bSparse)
+	: Handle(uint64(&Schema) | (uint64(OffsetWidth) << 1) | (uint64(bSparse) << KernelAddressBit) | uint64(1))
+	{
+		check(&Schema == &AsSchema());
+	}
+
+	bool						IsSchema() const		{ return !!(Handle & SchemaBit); }
+	bool						IsSparseSchema() const	{ return !!(Handle & KernelAddressBit); }
+	bool						IsMemcpy() const		{ return (Handle & NonSchemaMask) == MemcpyMask; }
+	bool						IsCustom() const		{ return (Handle & NonSchemaMask) == CustomMask; }
+	FLoadStructMemcpy			AsMemcpy() const		{ check(IsMemcpy()); return { static_cast<uint32>(Handle >> 32), static_cast<uint32>(Handle) >> 1 }; }
+	const ICustomStructBinding&	AsCustom() const		{ check(IsCustom()); return *AsPtr<ICustomStructBinding>(); }
+	const FStructSchemaBinding&	AsSchema() const		{ check(IsSchema()); return *AsPtr<FStructSchemaBinding>(); }
+
+private:
+	// This bit is always zero in user mode addresses and most likely won't be used by current or future
+	// CPU features like ARM's PAC / Top-Byte Ignore or Intel's Linear Address Masking / 5-Level Paging
+#if defined(__x86_64__) || defined(_M_X64)
+	static constexpr uint32 KernelAddressBit = 63;
+#elif defined(__aarch64__) || defined(_M_ARM64)
+	static constexpr uint32 KernelAddressBit = 55;
+#else
+	#error Unsupported architecture, please declare which address bit distinguish user space from kernel space
+#endif
+	// todo handle WASM, copy updated KernelAddressBit from AssetDataTagMap.h
+
+	static constexpr uint64 PtrMask = (uint64(1) << KernelAddressBit) | 0b111;
+	static constexpr uint64 NonSchemaMask = 0b11;
+	static constexpr uint64 MemcpyMask = 0b00;
+	static constexpr uint64 CustomMask = 0b10;
+	static constexpr uint64 SchemaBit = 1;
+	
+	template<typename T>
+	const T* AsPtr() const
+	{
+		check(Handle & PtrMask);
+		return reinterpret_cast<T*>(Handle & PtrMask);
+	}
+
+	uint64						Handle = 0;
+};
+
+////////////////////////////////////////////////////////////////////////////
+
+static bool ContainsEnums(const FStructSchemaBinding& Schema)
+{
+	for (FMemberBindType Member : MakeArrayView(Schema.Members, Schema.NumMembers))
+	{
+		if (Member.IsLeaf() && Member.AsLeaf().Bind.Type == ELeafBindType::Enum)
+		{
+			return true;
+		}
+	}
+	return false;
+}
+
+////////////////////////////////////////////////////////////////////////////
+
+struct FLoadBatch
+{
+	FReadBatchId			ReadId; // Needed to access schemas for custom struct loading
+	uint32					NumPlans;
+	FLoadStructPlan			Plans[0];
+
+	FLoadStructPlan			operator[](FStructSchemaId Id) const { check(Id.Idx < NumPlans); return Plans[Id.Idx]; }
+};
+
+using SubsetByteArray = TArray<uint8, TInlineAllocator<1024>>;
+
+static void CreateSubsetSchema(const FStructSchema& From, const FStructSchemaBinding& To, SubsetByteArray& Out)
+{
+	unimplemented();
+}
+
+static FLoadStructPlan CreateLoadPlan(const FStructSchema& From, const FStructSchemaBinding& To, SubsetByteArray& OutSubsetSchemas)
+{
+	// todo: detect if memcpyable, return FLoadStructMemcpy plan
+
+	if (From.NumMembers < To.NumMembers || ContainsEnums(To))
+	{
+		CreateSubsetSchema(From, To, /* out */ OutSubsetSchemas);
+	}
+	else
+	{
+		check(From.NumMembers == To.NumMembers);
+		check(From.NumInnerSchemas == To.NumInnerSchemas);
+		check(From.NumRangeTypes == To.NumInnerRanges);
+	}
+
+	// New subset schema pointers will be remapped later
+	return FLoadStructPlan(To, ELeafWidth::B32, From.IsDense);
+}
+
+FLoadBatch* CreateLoadPlans(FReadBatchId ReadId, const FDeclarations& Declarations, const FStructBindings& Bindings)
+{
+	// Temporary data structures
+	const uint32 NumPlans = NumStructSchemas(ReadId);
+	TArray<FLoadStructPlan, TInlineAllocator<256>> Plans;
+	TArray<uint32, TInlineAllocator<256>> SubsetSchemaSizes;
+	SubsetByteArray SubsetSchemaData;
+	Plans.SetNumUninitialized(NumPlans);
+	SubsetSchemaSizes.SetNumUninitialized(NumPlans);
+
+	// Create plans
+	for (FStructSchemaId Id = { 0 }; Id.Idx < NumPlans; ++Id.Idx)
+	{
+		FStructBinding Binding = Bindings.Get(Id);
+		int32 SubsetSchemaOffset = SubsetSchemaData.Num();
+		Plans[Id.Idx]	= Binding.IsCustom()
+						? FLoadStructPlan(Binding.AsCustom()) 
+						: CreateLoadPlan(ResolveStructSchema(ReadId, Id), Binding.AsSchema(), /* out */ SubsetSchemaData);
+		SubsetSchemaSizes[Id.Idx] = SubsetSchemaData.Num() - SubsetSchemaOffset;
+	}
+	
+	// Allocate load batch, copy plans and subset schemas, and fixup subset schema plans
+	SIZE_T Bytes = sizeof(FLoadBatch) + sizeof(FLoadStructPlan) * Plans.Num() + SubsetSchemaData.Num();
+	FLoadBatch Header = { ReadId, NumPlans };
+	FLoadBatch* Out = new (FMemory::Malloc(Bytes)) FLoadBatch{Header};
+	FMemory::Memcpy(Out->Plans, Plans.GetData(), sizeof(FLoadStructPlan) * Plans.Num());
+	if (SubsetSchemaData.Num() > 0)
+	{
+		uint8* OutSubsetData = reinterpret_cast<uint8*>(Out->Plans + Plans.Num());
+		FMemory::Memcpy(OutSubsetData, SubsetSchemaData.GetData(), SubsetSchemaData.Num());
+		
+		// Remap subset schema pointers
+		const uint8* It = OutSubsetData;
+		for (uint32 Idx = 0; Idx < NumPlans; ++Idx)
+		{
+			if (int32 Size = SubsetSchemaSizes[Idx])
+			{
+				check(IsAligned(Size, alignof(FStructSchemaBinding)));
+				check(Plans[Idx].IsSchema());
+				bool bSparse = Plans[Idx].IsSparseSchema();
+				Plans[Idx] = FLoadStructPlan(*reinterpret_cast<const FStructSchemaBinding*>(It), ELeafWidth::B32, bSparse);
+				It += Size;
+			}
+		}
+		check(It == OutSubsetData + SubsetSchemaData.Num());
+	}
+
+	return Out;
+}
+
+////////////////////////////////////////////////////////////////////////////
+
+FORCEINLINE static void SetBit(uint8& Out, uint8 Idx, bool bValue)
+{
+	uint8 Mask = IntCastChecked<uint8>(1 << Idx);
+	if (bValue)
+	{
+		Out |= Mask;
+	}
+	else
+	{
+		Out &= Mask;
+	}	
+}
+
+struct FLoadRangePlan
+{
+	ERangeSizeType MaxSize;
+	FOptionalStructSchemaId InnermostStruct;
+	TConstArrayView<FMemberBindType> InnerTypes;
+	const FRangeBinding* Bindings = nullptr;
+
+	FLoadRangePlan Tail() const
+	{
+		return { InnerTypes[0].AsRange().MaxSize, InnermostStruct, InnerTypes.RightChop(1), Bindings + 1 };
+	}
+};
+
+static FMemberBindType ToBindType(FMemberType Member)
+{
+	switch (Member.GetKind())
+	{
+	case EMemberKind::Leaf:		return FMemberBindType(Member.AsLeaf());
+	case EMemberKind::Range:	return FMemberBindType(Member.AsRange());
+	default:					return FMemberBindType(Member.AsStruct());	
+	}
+}
+
+class FRangeLoader
+{
+public:
+	static void LoadNestedRange(uint8* Member, FMemoryView Values, const FLoadBatch& Batch, const FLoadRangePlan& Range)
+	{
+		FByteReader ByteIt(Values);
+		FBitCacheReader BitIt;
+		LoadRange(Member, ByteIt, BitIt, Batch, Range);
+	}
+
+	static void LoadRangeView(uint8* Member, FRangeView Src, ERangeSizeType MaxSize, TConstArrayView<FRangeBinding> Bindings, const FLoadBatch& Batch)
+	{
+		TArray<FMemberBindType, TFixedAllocator<16>> InnerTypes;
+		InnerTypes.Add(ToBindType(Src.Schema.ItemType));
+		for (const FMemberType* It = Src.Schema.NestedItemTypes; It; It = It->IsRange() ? (It + 1) : nullptr)
+		{
+			InnerTypes.Add(ToBindType(*It));
+		}
+		check(Bindings.Num() == InnerTypes.Num());
+		
+		FOptionalStructSchemaId StructSchema = InnerTypes.Last().IsStruct() ? static_cast<FOptionalStructSchemaId>(Src.Schema.InnermostSchema) : NoId;
+		
+		FLoadRangePlan Plan = { MaxSize, StructSchema, InnerTypes, Bindings.GetData() };
+		
+		FByteReader ByteIt(Src.Values);
+		FBitCacheReader BitIt;
+		LoadRange(Member, ByteIt, BitIt, Batch, Plan);
+	}
+
+	static void LoadRange(uint8* Member, FByteReader& ByteIt, FBitCacheReader& BitIt, const FLoadBatch& Batch, const FLoadRangePlan& Range)
+	{
+		FMemberBindType InnerType = Range.InnerTypes[0];
+		const IRangeBinding& Binding = Range.Bindings[0].GetBinding();
+
+		if (uint64 Num = GrabRangeNum(Range.MaxSize, ByteIt, BitIt))
+		{
+			switch (InnerType.GetKind())
+			{
+				case EMemberKind::Leaf:		LoadRangeValues(Member, Num, Binding, ByteIt, Batch, InnerType.AsLeaf()); break;
+				case EMemberKind::Range:	LoadRangeValues(Member, Num, Binding, ByteIt, Batch, Range.Tail()); break;
+				case EMemberKind::Struct:	LoadRangeValues(Member, Num, Binding, ByteIt, Batch, Range.InnermostStruct.Get()); break;
+			}
+		}
+		else
+		{
+			FLoadRangeContext NoItemsCtx{.Request = {nullptr, 0}};
+			(Binding.MakeItems)(NoItemsCtx);
+		}
+	}
+
+	template<class SchemaType>
+	static void LoadRangeValues(uint8* Member, uint64 Num, const IRangeBinding& Binding, FByteReader& ByteIt, const FLoadBatch& Batch, SchemaType&& Schema)
+	{
+		FLoadRangeContext Ctx{.Request = {Member, Num}};
+		
+		while (Ctx.Request.Index < Num)
+		{
+			(Binding.MakeItems)(Ctx);
+			CopyRangeValues(Ctx.Items, ByteIt, Batch, Schema);
+			Ctx.Request.Index += Ctx.Items.Num;
+		}
+		
+		if (Ctx.Items.bNeedFinalize)
+		{
+			(Binding.MakeItems)(Ctx);
+		}
+	}
+	static void CopyRangeValues(const FConstructedItems& Items, FByteReader& ByteIt, const FLoadBatch& Batch, FLeafBindType Leaf)
+	{
+		switch (Leaf.Bind.Type)
+		{
+		case ELeafBindType::Bool:
+		{
+			check(Items.Size == sizeof(bool));
+			FBoolRangeView Bits(ByteIt.GrabBytes(Align(Items.Num, 8)/8), Items.Num);
+			uint8* It = Items.Data;
+			for (bool bBit : Bits)
+			{
+				reinterpret_cast<bool&>(*It++) = bBit;
+			}
+
+			break;
+		}
+		default:
+			checkf(Leaf.Bind.Type != ELeafBindType::BitfieldBool, TEXT("Loading to bitfield array is unsupported"));
+			check(Items.Size == SizeOf(Leaf.Arithmetic.Width));
+			FMemory::Memcpy(Items.Data, ByteIt.GrabBytes(Items.NumBytes()), Items.NumBytes());
+			break;
+		}
+	}
+	
+	static void CopyRangeValues(const FConstructedItems& Items, FByteReader& ByteIt, const FLoadBatch& Batch, FStructSchemaId Id)
+	{
+		uint64 ItemSize = Items.Size;
+		for (uint8* It = Items.Data, *End = It + Items.NumBytes(); It != End; It += ItemSize)
+		{
+			PlainProps::LoadStruct(It, FByteReader(ByteIt.GrabSkippableSlice()), Id, Batch);
+		}
+	}
+	
+	static void CopyRangeValues(const FConstructedItems& Items, FByteReader& ByteIt, const FLoadBatch& Batch, const FLoadRangePlan& Plan)
+	{
+		uint64 ItemSize = Items.Size;
+		for (uint8* It = Items.Data, *End = It + Items.NumBytes(); It != End; It += ItemSize)
+		{
+			LoadNestedRange(It, ByteIt.GrabSkippableSlice(), Batch, Plan);	
+		}
+	}
+};
+
+////////////////////////////////////////////////////////////////////////////
+
+template<bool bSparse, typename OffsetType>
+class TMemberLoader : public FRangeLoader
+{
+public:
+	TMemberLoader(FByteReader Values, const FStructSchemaBinding& Schema, const FLoadBatch& InBatch)
+	: Types(Schema.Members, Schema.NumMembers)
+	, Offsets(Schema.GetOffsets())
+	, InnerStructSchemas(static_cast<const FStructSchemaId*>(Schema.GetInnerSchemas()), Schema.NumInnerSchemas)
+	, InnerRangeTypes(Schema.GetInnerRangeTypes(), Schema.NumInnerRanges)
+	, RangeBindings(Schema.GetRangeBindings())
+	, Batch(InBatch)
+	, ByteIt(Values)
+	{
+		check(!ContainsEnums(Schema));
+	}
+
+	void Load(uint8* Struct)
+	{
+		SkipMissingSparseMembers();
+	
+		while (MemberIdx < Types.Num())
+		{
+			LoadMember(Struct);
+			++MemberIdx;
+			SkipMissingSparseMembers();
+		}
+	}
+
+private:
+	const TConstArrayView<FMemberBindType>		Types;
+	const OffsetType* const						Offsets;
+	const TConstArrayView<FStructSchemaId>		InnerStructSchemas;
+	const TConstArrayView<FMemberBindType>		InnerRangeTypes;
+	const FRangeBinding*						RangeBindings;
+	const FLoadBatch&							Batch;
+	
+	FByteReader ByteIt;
+	FBitCacheReader BitIt;
+	int32 MemberIdx = 0;
+	int32 InnerRangeIdx = 0;
+	int32 InnerStructIdx = 0;	
+
+	void SkipMissingSparseMembers()
+	{
+		// Make code changes in FMemberReader::SkipMissingSparseMembers() too
+		while (bSparse && MemberIdx < Types.Num() && BitIt.GrabNext(ByteIt))
+		{
+			FMemberBindType Type = Types[MemberIdx];
+			FMemberBindType InnermostType = Type.IsRange() ? GrabInnerRanges(Type.AsRange()).InnerTypes.Last() : Type;
+			if (InnermostType.IsStruct())
+			{
+				(void)GrabInnerStruct(InnermostType.AsStruct());
+			}
+			++MemberIdx;
+		}
+	}
+
+	void LoadMember(uint8* Struct)
+	{
+		FMemberBindType Type = Types[MemberIdx];
+		uint8* Member = Struct + Offsets[MemberIdx];
+
+		switch (Type.GetKind())
+		{
+			case EMemberKind::Leaf:		LoadMemberLeaf(Member, Type.AsLeaf()); break;
+			case EMemberKind::Range:	LoadMemberRange(Member, GrabInnerRanges(Type.AsRange())); break;
+			case EMemberKind::Struct:	LoadMemberStruct(Member, GrabInnerStruct(Type.AsStruct())); break;
+		}
+	}
+
+	FStructSchemaId GrabInnerStruct(FStructBindType Type)
+	{
+		return Type.IsDynamic ? FStructSchemaId { ByteIt.Grab<uint32>() } : InnerStructSchemas[InnerStructIdx++];
+	}
+
+	FLoadRangePlan GrabInnerRanges(FRangeBindType Type)
+	{
+		const FRangeBinding* Bindings = RangeBindings + InnerRangeIdx;
+		TConstArrayView<FMemberBindType> InnerTypes = GrabInnerRangeTypes(InnerRangeTypes, /* in-out */ InnerRangeIdx);	
+		FOptionalStructSchemaId InnermostStruct = InnerTypes.Last().IsStruct() ? ToOptional(GrabInnerStruct(InnerTypes.Last().AsStruct())) : NoId;
+		return { Type.MaxSize, InnermostStruct, InnerTypes, Bindings };
+	}
+	
+	void LoadMemberLeaf(uint8* Member, FLeafBindType Leaf)
+	{
+		switch (Leaf.Bind.Type)
+		{
+			case ELeafBindType::Bool:
+				reinterpret_cast<bool&>(*Member) = BitIt.GrabNext(ByteIt);
+				break;
+			case ELeafBindType::BitfieldBool:
+				SetBit(*Member, Leaf.Bitfield.Idx, BitIt.GrabNext(ByteIt));
+				break;
+			default:
+				switch (Leaf.Arithmetic.Width)
+				{
+					case ELeafWidth::B8:	FMemory::Memcpy(Member, ByteIt.GrabBytes(1), 1); break;
+					case ELeafWidth::B16:	FMemory::Memcpy(Member, ByteIt.GrabBytes(2), 2); break;
+					case ELeafWidth::B32:	FMemory::Memcpy(Member, ByteIt.GrabBytes(4), 4); break;
+					case ELeafWidth::B64:	FMemory::Memcpy(Member, ByteIt.GrabBytes(8), 8); break;
+				}
+				break;
+		}
+	}
+
+	void LoadMemberStruct(uint8* Member, FStructSchemaId Id)
+	{
+		PlainProps::LoadStruct(Member, FByteReader(ByteIt.GrabSkippableSlice()), Id, Batch);
+	}
+
+	void LoadMemberRange(uint8* Member, FLoadRangePlan&& Range)
+	{
+		LoadRange(Member, ByteIt, BitIt, Batch, Range);
+	}
+};
+
+////////////////////////////////////////////////////////////////////////////
+
+void LoadStruct(uint8* Dst, FByteReader Src, FStructSchemaId Id, const FLoadBatch& Batch)
+{
+	FLoadStructPlan Plan = Batch[Id];
+	if (Plan.IsSchema())
+	{
+		if (Plan.IsSparseSchema())
+		{
+			TMemberLoader< true, uint32>(Src, Plan.AsSchema(), Batch).Load(Dst);
+		}
+		else
+		{
+			TMemberLoader<false, uint32>(Src, Plan.AsSchema(), Batch).Load(Dst);
+		}
+	}
+	else if (Plan.IsMemcpy())
+	{
+		Src.CheckSize(Plan.AsMemcpy().Size);
+		FMemory::Memcpy(Dst + Plan.AsMemcpy().Offset, Src.Peek(), Plan.AsMemcpy().Size);
+	}
+	else
+	{
+		FStructSchemaHandle ReadSchema{Id, Batch.ReadId};
+		Plan.AsCustom().LoadStruct(Dst, { ReadSchema, Src }, ECustomLoadMethod::Assign, Batch);
+	}
+}
+
+void LoadStruct(uint8* Dst, FStructView Src, const FLoadBatch& Batch)
+{
+	LoadStruct(Dst, Src.Values, Src.Schema.Id, Batch);
+}
+
+void ConstructAndLoadStruct(uint8* Dst, FByteReader Src, FStructSchemaId Id, const FLoadBatch& Batch)
+{
+	FLoadStructPlan Plan = Batch[Id];
+	checkf(!Plan.IsSchema(), TEXT("Non-default constructible types requires ICustomStructBinding or in rare cases memcpying"));
+
+	if (Plan.IsMemcpy())
+	{
+		Src.CheckSize(Plan.AsMemcpy().Size);
+		FMemory::Memcpy(Dst + Plan.AsMemcpy().Offset, Src.Peek(), Plan.AsMemcpy().Size);
+	}
+	else
+	{
+		FStructSchemaHandle ReadSchema{Id, Batch.ReadId};
+		Plan.AsCustom().LoadStruct(Dst, { ReadSchema, Src }, ECustomLoadMethod::Construct, Batch);
+	}
+}
+
+void LoadRange(uint8* Dst, FRangeView Src, ERangeSizeType MaxSize, TConstArrayView<FRangeBinding> Bindings, const FLoadBatch& Batch)
+{
+	FRangeLoader::LoadRangeView(Dst, Src, MaxSize, Bindings, Batch);
+}
+
+
+} // namespace PlainProps
\ No newline at end of file
diff --git a/Engine/Source/Programs/PlainPropsTests/Private/PlainPropsRead.cpp b/Engine/Source/Programs/PlainPropsTests/Private/PlainPropsRead.cpp
new file mode 100644
--- /dev/null
+++ b/Engine/Source/Programs/PlainPropsTests/Private/PlainPropsRead.cpp
@@ -0,0 +1,446 @@
+// Copyright Epic Games, Inc. All Rights Reserved.
+
+#include  "PlainPropsRead.h"
+#include  "PlainPropsInternalFormat.h"
+#include  "PlainPropsInternalRead.h"
+#include "Misc/Optional.h"
+
+#include <atomic>
+#include <type_traits>
+
+namespace PlainProps
+{
+
+void FSchemaBatch::ValidateBounds(uint64 NumBytes) const
+{
+	static constexpr uint32 Alignment = FMath::Max(alignof(FStructSchema), alignof(FEnumSchema));
+	check(IsAligned(this, Alignment));
+	check(sizeof(FSchemaBatch) + NumSchemas * sizeof(uint32) <= NestedScopesOffset);
+	check(NestedScopesOffset + NumNestedScopes * sizeof(FNestedScope) + NumParametricTypes * sizeof(FParametricType) <= NumBytes);
+
+	for (uint32 SchemaOffset : MakeArrayView(SchemaOffsets, NumSchemas))
+	{
+		check(SchemaOffset < NestedScopesOffset);
+		check(IsAligned(SchemaOffset, Alignment));
+	}
+		
+	uint32 NumParameters = 0;
+	for (FParametricType ParametricType : GetParametricTypes())
+	{
+		check(ParametricType.Parameters.Idx == NumParameters);
+		check(ParametricType.Parameters.NumParameters > 0);
+		NumParameters += ParametricType.Parameters.NumParameters;
+	}
+
+	const void* ExpectedEnd = GetFirstParameter() + NumParameters;
+	const void* ActualEnd = reinterpret_cast<const uint8*>(this) + NumBytes;
+	check(ExpectedEnd == ActualEnd);
+}
+
+template<class SchemaType>
+const SchemaType& ResolveSchema(const FSchemaBatch& Batch, FSchemaId Id)
+{
+	check(Id.Idx < Batch.NumSchemas);
+	return *reinterpret_cast<const SchemaType*>(reinterpret_cast<const uint8*>(&Batch) + Batch.SchemaOffsets[Id.Idx]);
+}
+	
+static FNestedScope ResolveNestedScope(const FSchemaBatch& Batch, FNestedScopeId Id)
+{
+	return Batch.GetNestedScopes()[Id.Idx];
+}
+	
+static FParametricTypeView ResolveParametricType(const FSchemaBatch& Batch, FParametricTypeId Id)
+{
+	FParametricType Type = Batch.GetParametricTypes()[Id.Idx];
+	return {Type.Name, IntCastChecked<uint8>(Type.Parameters.NumParameters), Batch.GetFirstParameter() + Type.Parameters.Idx};
+}
+
+
+class FReadSchemaRegistry
+{
+	static constexpr uint32 Capacity = 1 << 16;
+	std::atomic<const FSchemaBatch*> Slots[Capacity]{};
+	std::atomic<uint32> Counter{};
+
+public:
+	FReadBatchId Mount(const FSchemaBatch* Batch)
+	{
+		const uint32 Count = Counter.fetch_add(1, std::memory_order_relaxed);
+		for (uint32 BaseIdx = 0; BaseIdx < FReadSchemaRegistry::Capacity; ++BaseIdx)
+		{
+			uint32 Idx = (Count + BaseIdx) % FReadSchemaRegistry::Capacity;
+			const FSchemaBatch* Null = nullptr;
+			if (Slots[Idx].load(std::memory_order_relaxed) == nullptr &&
+				Slots[Idx].compare_exchange_strong(Null, Batch, std::memory_order_acq_rel))
+			{
+				FReadBatchId Out;
+				Out.Idx = static_cast<uint16>(Idx);
+				return Out;
+			}
+		}
+
+		checkf(false, TEXT("Exceeded fixed limit of 65536 simulatenous read batches"));
+		while (true);
+	}
+
+	const FSchemaBatch* Unmount(FReadBatchId Id)
+	{
+		const FSchemaBatch* Batch = &Get(Id);
+		verify(Slots[Id.Idx].compare_exchange_strong(Batch, nullptr, std::memory_order_acq_rel));
+		return Batch;
+	}
+
+	const FSchemaBatch& Get(FReadBatchId Id) const
+	{
+		const FSchemaBatch* Batch = Slots[Id.Idx].load(std::memory_order_relaxed);
+		check(Batch);
+		return *Batch;
+	}
+};
+static FReadSchemaRegistry GReadSchemas;
+
+const FSchemaBatch*	ValidateSchemas(FMemoryView Schemas)
+{
+	const FSchemaBatch* Batch = reinterpret_cast<const FSchemaBatch*>(Schemas.GetData());
+	Batch->ValidateBounds(Schemas.GetSize());
+	return Batch;
+}
+
+FReadBatchId MountReadSchemas(const FSchemaBatch* Batch)
+{
+	return GReadSchemas.Mount(Batch);
+}
+
+const FSchemaBatch* UnmountReadSchemas(FReadBatchId Id)
+{
+	return GReadSchemas.Unmount(Id);
+}
+
+uint32 NumStructSchemas(FReadBatchId Batch)
+{
+	return GReadSchemas.Get(Batch).NumSchemas;
+}
+
+const FStructSchema& ResolveStructSchema(FReadBatchId Batch, FStructSchemaId Schema)
+{
+	return ResolveSchema<FStructSchema>(GReadSchemas.Get(Batch), Schema);
+}
+
+const FEnumSchema& ResolveEnumSchema(FReadBatchId Batch, FEnumSchemaId Schema)
+{
+	return ResolveSchema<FEnumSchema>(GReadSchemas.Get(Batch), Schema);
+}
+
+FNestedScope ResolveUntranslatedNestedScope(FReadBatchId Batch, FNestedScopeId Id)
+{
+	return ResolveNestedScope(GReadSchemas.Get(Batch), Id);
+}
+
+FParametricTypeView ResolveUntranslatedParametricType(FReadBatchId Batch, FParametricTypeId Id)
+{
+	return ResolveParametricType(GReadSchemas.Get(Batch), Id);
+}
+
+//////////////////////////////////////////////////////////////////////////
+
+FLeafRangeView FRangeView::AsLeaves() const
+{
+	FUnpackedLeafType Leaf = Schema.ItemType.AsLeaf();
+
+	FLeafRangeView Out;
+	Out.Type = Leaf.Type;
+	Out.Width = Leaf.Width;
+	Out.Batch = Schema.Batch;
+	Out.Enum = static_cast<FOptionalEnumSchemaId>(Schema.InnermostSchema);
+	Out.NumItems = NumItems;
+	Out.Values = static_cast<const uint8*>(Values.GetData());
+
+	return Out;
+}
+
+FStructRangeView FRangeView::AsStructs() const
+{
+	check(IsStructRange());
+
+	FStructRangeView Out;
+	Out.NumItems = NumItems;
+	Out.Data = Values;
+	Out.Schema.Id = static_cast<FStructSchemaId>(Schema.InnermostSchema.Get());
+	Out.Schema.Batch = Schema.Batch;
+
+	return Out;
+}
+
+FNestedRangeView FRangeView::AsRanges() const
+{
+	check(IsNestedRange());
+
+	FNestedRangeView Out;
+	Out.NumItems = NumItems;
+	Out.Data = Values;
+	Out.Schema = Schema;
+	return Out;
+}
+
+//////////////////////////////////////////////////////////////////////////
+
+bool FBitCacheReader::GrabNext(FByteReader& Bytes)
+{
+	BitIt <<= 1; // Shift up til overflow
+
+	if (BitIt == 0)
+	{
+		Bits = Bytes.GrabByte();
+		BitIt = 1;
+	}
+
+	return !!(Bits & BitIt);
+}
+
+const FStructSchema& FStructSchemaHandle::ResolveSuper() const
+{
+	return ResolveStructSchema(Batch, Resolve().GetSuperSchema().Get());
+}
+
+//////////////////////////////////////////////////////////////////////////
+
+static FMemoryView GrabRangeValues(uint64 Num, FMemberType InnerType, /* in-out */ FByteReader& ByteIt)
+{
+	if (Num == 0)
+	{
+		return FMemoryView();
+	}
+
+	uint64 NumBytes = InnerType.GetKind() == EMemberKind::Leaf ? GetLeafRangeSize(Num, InnerType.AsLeaf()) : ByteIt.GrabVarIntU();
+
+	return ByteIt.GrabSlice(NumBytes);
+}
+
+FRangeView FNestedRangeIterator::operator*() const 
+{
+	FByteReader PeekBytes = ByteIt;
+	FBitCacheReader PeekBits = BitIt;
+	
+	FRangeView Out;
+	Out.Schema = { Schema.NestedItemTypes[0], Schema.Batch, Schema.InnermostSchema, /* Only valid for nested ranges */ Schema.NestedItemTypes + 1 };
+	Out.NumItems = GrabRangeNum(Schema.ItemType.AsRange().MaxSize, /* in-out */ PeekBytes,  /* in-out */ PeekBits);
+	Out.Values = GrabRangeValues(Out.NumItems, Schema.NestedItemTypes[0], /* in-out */ PeekBytes);
+
+	return Out;
+}
+
+void FNestedRangeIterator::operator++()
+{
+	uint64 Num = GrabRangeNum(Schema.ItemType.AsRange().MaxSize, /* in-out */ ByteIt, /* in-out */ BitIt);
+	(void)GrabRangeValues(Num, Schema.NestedItemTypes[0], /* in-out */ ByteIt);
+}
+
+//////////////////////////////////////////////////////////////////////////
+
+const FMemberType*	FMemberReader::GetMemberTypes() const	{ return FStructSchema::GetMemberTypes(Footer); }
+const FMemberType*	FMemberReader::GetRangeTypes() const	{ return FStructSchema::GetRangeTypes(Footer, NumMembers); }
+const FSchemaId*	FMemberReader::GetInnerSchemas() const	{ return FStructSchema::GetInnerSchemas(Footer, NumMembers, NumRangeTypes, NumMembers - HasSuper); }
+const FMemberId*	FMemberReader::GetMemberNames() const	{ return FStructSchema::GetMemberNames(Footer, NumMembers, NumRangeTypes); }
+
+static bool SkipDeclaredSuperSchema(ESuper Inheritance)
+{
+	return Inheritance == ESuper::Unused || Inheritance == ESuper::Used;
+}
+
+FMemberReader::FMemberReader(const FStructSchema& Schema, FByteReader Values, FReadBatchId InBatch)
+: Footer(Schema.Footer)
+, Batch(InBatch)
+, NumMembers(Schema.NumMembers)
+, NumRangeTypes(Schema.NumRangeTypes)
+, IsSparse(!Schema.IsDense)
+, HasSuper(UsesSuper(Schema.Inheritance))
+, InnerSchemaIdx(SkipDeclaredSuperSchema(Schema.Inheritance))
+, ValueIt(Values)
+#if DO_CHECK
+, NumInnerSchemas(Schema.NumInnerSchemas)
+#endif
+{
+	check(InnerSchemaIdx <= NumInnerSchemas);
+	checkf(NumRangeTypes != 0xFFFFu, TEXT("GrabRangeTypes() doesn't check for wrap-around"));
+	SkipMissingSparseMembers();
+}
+
+FOptionalMemberId FMemberReader::PeekName() const
+{
+	int32 MemberNameIdx = MemberIdx - HasSuper;
+	if (MemberNameIdx == -1)
+	{
+		UE_DEBUG_BREAK();
+	}
+
+	return MemberNameIdx >= 0 ? ToOptional(GetMemberNames()[MemberNameIdx]) : NoId;
+}
+
+EMemberKind FMemberReader::PeekKind() const
+{
+	return PeekType().GetKind();
+}
+
+FMemberType	FMemberReader::PeekType() const
+{
+	check(HasMore());
+	return GetMemberTypes()[MemberIdx];
+}
+
+void FMemberReader::AdvanceToNextMember()
+{
+	++MemberIdx;
+	SkipMissingSparseMembers();
+}
+
+void FMemberReader::SkipMissingSparseMembers()
+{
+	if (IsSparse)
+	{
+		// Change code in LoadMembers() too
+		while (MemberIdx < NumMembers && GrabBit())
+		{
+			FMemberType InnermostType = PeekType().IsRange() ? /* advances RangeTypeIdx */ GrabRangeTypes().Last() : PeekType();
+			SkipSchema(InnermostType);
+			++MemberIdx;
+		}
+	}	
+}
+
+void FMemberReader::SkipSchema(FMemberType InnermostType)
+{
+	if (InnermostType.IsStruct())
+	{
+		if (InnermostType.AsStruct().IsDynamic)
+		{
+			(void)ValueIt.Grab<uint32>();
+		}
+		else
+		{
+			++InnerSchemaIdx;
+		}
+	}
+	else
+	{
+		InnerSchemaIdx += InnermostType.AsLeaf().Type == ELeafType::Enum;
+	}
+
+	check(InnerSchemaIdx <= NumInnerSchemas);
+}
+
+FSchemaId FMemberReader::GrabInnerSchema()
+{
+	check(InnerSchemaIdx < NumInnerSchemas);
+	uint32 Idx = InnerSchemaIdx++;
+	return GetInnerSchemas()[Idx];
+}
+
+FStructSchemaId FMemberReader::GrabStructSchema(FStructType Type)
+{
+	return Type.IsDynamic	? FStructSchemaId { ValueIt.Grab<uint32>() }
+							: static_cast<FStructSchemaId&&>(GrabInnerSchema());
+}
+
+FOptionalSchemaId FMemberReader::GrabRangeSchema(FMemberType InnermostType)
+{
+	if (InnermostType.IsStruct())
+	{
+		return FOptionalSchemaId(GrabStructSchema(InnermostType.AsStruct()));
+	}
+
+	ELeafType LeafType = InnermostType.AsLeaf().Type;
+	return LeafType == ELeafType::Enum ? FOptionalSchemaId(GrabEnumSchema()) : NoId;
+}
+
+FLeafView FMemberReader::GrabLeaf()
+{
+	check(HasMore());
+	FUnpackedLeafType Leaf = PeekType().AsLeaf();
+	FEnumSchemaId Enum = Leaf.Type == ELeafType::Enum ? GrabEnumSchema() : FEnumSchemaId{};
+	
+	FMemberValue Value;
+	if (Leaf.Type == ELeafType::Bool)
+	{
+		Value.bValue = GrabBit();
+	}
+	else
+	{
+		Value.Ptr = ValueIt.GrabBytes(SizeOf(Leaf.Width));
+	}
+
+	AdvanceToNextMember();
+
+	return { Leaf.Type, Leaf.Width, Batch, Enum, Value };
+}
+
+FStructView FMemberReader::GrabStruct()
+{
+	check(HasMore());
+	FStructType Type = PeekType().AsStruct();
+	FStructSchemaId Struct = Type.IsDynamic ? FStructSchemaId{ ValueIt.Grab<uint32>() } : GrabStructSchema(Type);
+	FMemoryView Values = ValueIt.GrabSkippableSlice();
+
+	AdvanceToNextMember();
+
+	return { { Struct, Batch }, FByteReader(Values) };
+}
+
+TConstArrayView<FMemberType> FMemberReader::GrabRangeTypes()
+{
+	return GrabInnerRangeTypes(MakeArrayView(GetRangeTypes(), NumRangeTypes), /* in-out */ RangeTypeIdx);
+}
+
+FRangeView FMemberReader::GrabRange()
+{
+	check(HasMore());
+
+	FMemberTypeRange RangeTypes = GrabRangeTypes();
+	FRangeView Out;
+	Out.Schema.Batch = Batch;
+	Out.Schema.InnermostSchema = GrabRangeSchema(RangeTypes.Last());									// #1
+	Out.Schema.ItemType = RangeTypes[0];
+	Out.Schema.NestedItemTypes = RangeTypes.Num() > 1 ? &RangeTypes[1] : nullptr;
+	Out.NumItems = GrabRangeNum(PeekType().AsRange().MaxSize, /* in-out */ ValueIt, /* in-out */ Bits); // #2
+	Out.Values = GrabRangeValues(Out.NumItems, Out.Schema.ItemType, /* in-out */ ValueIt);				// #3
+	
+	AdvanceToNextMember();
+
+	return Out;
+}
+
+//FAnyMemberView FMemberReader::GrabAnyMember()
+//{
+//	FMemberValue Value = GetValueAndAdvance(Type, /* in-out */ ValueIt);
+//	
+//	return {Type, Id, BatchId, Value};
+//}
+
+//////////////////////////////////////////////////////////////////////////
+
+static TOptional<FStructView> TryGrabSuper(FMemberReader& Members)
+{
+	if (Members.HasMore() && IsSuper(Members.PeekType()))
+	{
+		return Members.GrabStruct();
+	}
+
+	return NullOpt;
+}
+
+FFlatMemberReader::FReader::FReader(FStructView Struct)
+: FMemberReader(Struct)
+, Owner(Struct.Schema.Resolve().Type)
+{}
+
+
+FFlatMemberReader::FFlatMemberReader(FStructView Struct)
+{
+	Lineage.Emplace(Struct);
+	while (TOptional<FStructView> Super = TryGrabSuper(Lineage.Last()))
+	{
+		Lineage.Emplace(*Super);
+	}
+	It = &Lineage.Last();
+}
+
+} // namespace PlainProps
\ No newline at end of file
diff --git a/Engine/Source/Programs/PlainPropsTests/Private/PlainPropsSave.cpp b/Engine/Source/Programs/PlainPropsTests/Private/PlainPropsSave.cpp
new file mode 100644
--- /dev/null
+++ b/Engine/Source/Programs/PlainPropsTests/Private/PlainPropsSave.cpp
@@ -0,0 +1,366 @@
+// Copyright Epic Games, Inc. All Rights Reserved.
+
+#include  "PlainPropsSave.h"
+#include  "PlainPropsBind.h"
+#include  "PlainPropsInternalBuild.h"
+#include  "PlainPropsInternalFormat.h"
+#include <type_traits>
+
+namespace PlainProps
+{
+
+static uint64 GetBit(uint8 Byte, uint8 BitIdx)
+{
+	return (Byte >> BitIdx) & 1;
+}
+
+static uint64 SaveLeaf(const uint8* Member, FUnpackedLeafBindType Leaf)
+{
+#if DO_CHECK
+	if (Leaf.Type == ELeafBindType::Float)
+	{
+		switch (Leaf.Width)
+		{
+			case ELeafWidth::B32: return ValueCast(*reinterpret_cast<const float*>(Member));
+			case ELeafWidth::B64: return ValueCast(*reinterpret_cast<const double*>(Member));
+			default: check(false); return 0;
+		}
+
+	}
+#endif
+
+	if (Leaf.Type == ELeafBindType::BitfieldBool)
+	{
+		return GetBit(*Member, Leaf.BitfieldIdx);
+	}
+	else
+	{
+		// Undefined behavior. If problematic, use memcpy or switch(Leaf.Type) and ValueCast(reinterpret_cast<...>(Member))  
+		switch (Leaf.Width)
+		{
+			case ELeafWidth::B8:	return *Member;
+			case ELeafWidth::B16:	return *reinterpret_cast<const uint16*>(Member);
+			case ELeafWidth::B32:	return *reinterpret_cast<const uint32*>(Member);
+			case ELeafWidth::B64:	return *reinterpret_cast<const uint64*>(Member);
+		}
+	}
+	check(false);
+	return uint64(0);
+}
+
+//FBuiltRange* SaveLeafRange(const void* Range, const IRangeBinding& Binding, FUnpackedLeafType Leaf)
+//{
+//	FSaveRangeContext Ctx = { { Range } };
+//
+//	Binding.ReadItems(Ctx);
+//
+//	if (const uint64 NumTotal = Ctx.Items.NumTotal)
+//	{
+//		FBuiltRange* Out = new (FMemory::Malloc(sizeof(FBuiltRange) + NumTotal * SizeOf(Leaf.Width))) FBuiltRange;
+//		Out->Num = NumTotal;
+//
+//		for (uint64 NumRead = 0; true; NumRead += Ctx.Items.Slice.Num)
+//		{
+//			check(Ctx.Items.Slice.Num > 0);
+//
+//			FMemory::Memcpy(Out->Data + NumRead, Ctx.Items.Slice.Data,  Ctx.Items.Slice.Num * SizeOf(Leaf.Width));
+//
+//			if (NumRead >= NumTotal)
+//			{
+//				check(NumRead == NumTotal);	
+//				break;
+//			}
+//
+//			Binding.ReadItems(Ctx);	
+//		}
+//
+//		return Out;
+//	}
+//	
+//	return nullptr;
+//}
+//
+//FBuiltRange* SaveStructRange(const void* Range, const IRangeBinding& Binding, FStructType Type, FStructSchemaId Id, const FSaveContext& OuterCtx)
+//{
+//	FSaveRangeContext Ctx = { { Range } };
+//
+//	Binding.ReadItems(Ctx);
+//
+//	if (const uint64 NumTotal = Ctx.Items.NumTotal)
+//	{
+//		TArray64<TUniquePtr<FBuiltStruct>> Structs;
+//		Structs.Reserve(NumTotal);
+//
+//		for (uint64 NumRead = 0; true; NumRead += Ctx.Items.Slice.Num)
+//		{
+//			check(Ctx.Items.Slice.Num > 0);
+//
+//			for (uint64 Idx = 0, Num = Ctx.Items.Slice.Num; Idx < Num; ++Idx)
+//			{
+//				Structs.Emplace(SaveStruct(Ctx.Items.Slice.At(Idx, Ctx.Items.Stride), Id, OuterCtx));
+//			}
+//
+//			if (NumRead >= NumTotal)
+//			{
+//				check(NumRead == NumTotal);	
+//				break;
+//			}
+//
+//			Binding.ReadItems(Ctx);	
+//		}
+//
+//		return Private::BuildStructRange(/* ownership xfer */ Structs);
+//	}
+//	
+//	return nullptr;
+//}
+
+
+
+//
+//static FBuiltRange* SaveNestedRange(const void* Range, const IRangeBinding& Binding, FRangeType Type, FRangeMemberBinding InnerBinding, const FSaveContext& OuterCtx)
+//{
+//	FSaveRangeContext Ctx = { { Range } };
+//
+//	Binding.ReadItems(Ctx);
+//
+//	if (const uint64 NumTotal = Ctx.Items.NumTotal)
+//	{
+//		TArray64<FBuiltRange*> Ranges;
+//		Ranges.Reserve(NumTotal);
+//
+//		for (uint64 NumRead = 0; true; NumRead += Ctx.Items.Slice.Num)
+//		{
+//			check(Ctx.Items.Slice.Num > 0);
+//
+//			for (uint64 Idx = 0, Num = Ctx.Items.Slice.Num; Idx < Num; ++Idx)
+//			{
+//				Ranges.Emplace(SaveRange(Ctx.Items.Slice.At(Idx, Ctx.Items.Stride), InnerBinding, OuterCtx));
+//			}
+//
+//			if (NumRead >= NumTotal)
+//			{
+//				check(NumRead == NumTotal);	
+//				break;
+//			}
+//
+//			Binding.ReadItems(Ctx);	
+//		}
+//
+//		return Private::BuildNestedRange(/* ownership xfer */ Ranges);
+//	}
+//	
+//	return nullptr;
+//}
+
+
+//struct FStructRangeSaver
+//{
+//	FStructSchemaId Id;
+//	TArray64<TUniquePtr<FBuiltStruct>> Structs;
+//
+//	void Init(uint64 Num)
+//	{ 
+//		Ranges.Reserve(NumTotal);
+//	}
+//
+//	void Append(FExistingItemSlice Slice, uint32 Stride, FSaveContext& OuterCtx)
+//	{
+//		for (uint64 Idx = 0; Idx < Slice.Num; ++Idx)
+//		{
+//			Ranges.Emplace(SaveRange(Slice.At(Idx, Stride), Id, OuterCtx));
+//		}
+//	}
+//
+//	FBuiltRange* Build()
+//	{
+//		Private::BuildStructRange(/* ownership xfer */ Structs);
+//	}
+//};
+
+struct FLeafRangeSaver
+{
+	FBuiltRange* Out;
+	uint8* OutIt;
+
+	FLeafRangeSaver(uint64 Num, SIZE_T LeafSize)
+	: Out(new (FMemory::Malloc(sizeof(FBuiltRange) + Num * LeafSize)) FBuiltRange)
+	, OutIt(Out->Data)
+	{ 
+		Out->Num = Num;
+	}
+
+	void Append(FExistingItemSlice Slice, uint32, SIZE_T LeafSize, const FSaveContext&)
+	{
+		FMemory::Memcpy(OutIt, Slice.Data, Slice.Num * LeafSize);
+		OutIt += Slice.Num * LeafSize;
+	}
+
+	FBuiltRange* Finish()
+	{
+		return Out;
+	}
+};
+
+template<typename BuiltItemType, typename ItemSchemaType>
+struct TStructuralRangeSaver
+{
+	TArray64<BuiltItemType> Items;
+
+	TStructuralRangeSaver(uint64 Num, ItemSchemaType)
+	{ 
+		Items.Reserve(Num);
+	}
+
+	void Append(FExistingItemSlice Slice, uint32 Stride, ItemSchemaType Schema, const FSaveContext& OuterCtx)
+	{
+		for (uint64 Idx = 0; Idx < Slice.Num; ++Idx)
+		{
+			Items.Emplace(SaveRangeItem(Slice.At(Idx, Stride), Schema, OuterCtx));
+		}
+	}
+
+	[[nodiscard]] FBuiltRange* Finish()
+	{
+		return Private::BuildStructuralRange(/* ownership xfer */ Items);
+	}
+};
+
+using FNestedRangeSaver = TStructuralRangeSaver<FBuiltRange*, FRangeMemberBinding>;
+using FStructRangeSaver = TStructuralRangeSaver<TUniquePtr<FBuiltStruct>, FStructSchemaId>;
+
+template<class SaverType, typename InnerContextType>
+[[nodiscard]] FBuiltRange* SaveRange(const void* Range, const IRangeBinding& Binding, const FSaveContext& OuterCtx, InnerContextType InnerCtx)
+{
+	FSaveRangeContext ReadCtx = { { Range } };
+	Binding.ReadItems(ReadCtx);
+
+	if (const uint64 NumTotal = ReadCtx.Items.NumTotal)
+	{
+		SaverType Saver(NumTotal, InnerCtx);
+		while (true)
+		{
+			check(ReadCtx.Items.Slice.Num > 0);
+			Saver.Append(ReadCtx.Items.Slice, ReadCtx.Items.Stride, InnerCtx, OuterCtx);
+		
+			ReadCtx.Request.Index += ReadCtx.Items.Slice.Num;
+			if (ReadCtx.Request.Index >= NumTotal)
+			{
+				check(ReadCtx.Request.Index == NumTotal);	
+				return Saver.Finish();
+			}
+
+			Binding.ReadItems(ReadCtx);	
+		}
+	}
+	
+	return nullptr;
+}
+
+[[nodiscard]] static FRangeMemberBinding GetInnerRange(FRangeMemberBinding Member)
+{
+	check(Member.InnerTypes[0].IsRange());
+	return { Member.InnerTypes + 1, Member.RangeBindings + 1, Member.InnermostSchema };
+}
+
+ ELeafWidth GetArithmeticWidth(FLeafBindType Leaf)
+ {
+	checkf(Leaf.Bind.Type != ELeafBindType::BitfieldBool, TEXT("Arrays of bitfields is not a thing"));
+	return Leaf.Arithmetic.Width;
+ }
+
+[[nodiscard]] static FBuiltRange* SaveRange(const uint8* Range, FRangeMemberBinding Member, const FSaveContext& Ctx)
+{
+	const IRangeBinding& Binding = Member.RangeBindings[0].GetBinding();
+	FMemberBindType InnerType = Member.InnerTypes[0];
+	switch (InnerType.GetKind())
+	{
+	case EMemberKind::Leaf:		return SaveRange<FLeafRangeSaver>(  Range, Binding, Ctx, SizeOf(GetArithmeticWidth(InnerType.AsLeaf())));
+	case EMemberKind::Range:	return SaveRange<FNestedRangeSaver>(Range, Binding, Ctx, GetInnerRange(Member));
+	case EMemberKind::Struct:	return SaveRange<FStructRangeSaver>(Range, Binding, Ctx, static_cast<FStructSchemaId>(Member.InnermostSchema.Get()));
+	}
+
+	check(false);
+	return nullptr;
+}
+
+[[nodiscard]] static FBuiltRange* SaveRangeItem(const uint8* Range, FRangeMemberBinding Member, const FSaveContext& Ctx)
+{ 
+	return SaveRange(Range, Member, Ctx);
+}
+
+[[nodiscard]] static TUniquePtr<FBuiltStruct> SaveRangeItem(const uint8* Struct, FStructSchemaId Id, const FSaveContext& Ctx)
+{
+	return SaveStruct(Struct, Id, Ctx);
+}
+
+[[nodiscard]] static FMemberType ToMemberType(FMemberBindType In)
+{
+	switch (In.GetKind())
+	{
+	case EMemberKind::Leaf:		return FMemberType(ToLeafType(In.AsLeaf()));
+	case EMemberKind::Range:	return FMemberType(In.AsRange());
+	default:					return FMemberType(In.AsStruct());
+	}
+}
+
+[[nodiscard]] static TArray<FMemberType> ToMemberTypes(const FMemberBindType* InnerTypes)
+{
+	TArray<FMemberType> Out;
+	for (const FMemberBindType* It = InnerTypes; true; ++It)
+	{
+		Out.Add(ToMemberType(*It));
+		if (!It->IsRange())
+		{
+			return Out;
+		}
+	} 
+}
+
+[[nodiscard]] static FMemberSchema MakeSchema(FRangeMemberBinding Member)
+{
+	return { FMemberType(Member.RangeBindings[0].GetSizeType()), Member.InnermostSchema, ToMemberTypes(Member.InnerTypes) };
+}
+
+
+static void SaveMember(FMemberBuilder& Out, const uint8* Struct, FMemberId Name, const FSaveContext& Ctx, FLeafMemberBinding Member)
+{
+	FUnpackedLeafType Type = { ToLeafType(Member.Leaf.Type), Member.Leaf.Width };
+	Out.AddLeaf(Name, Type, Member.Enum, SaveLeaf(Struct + Member.Offset, Member.Leaf));
+}
+
+static void SaveMember(FMemberBuilder& Out, const uint8* Struct, FMemberId Name, const FSaveContext& Ctx, FRangeMemberBinding Member)
+{
+	Out.AddRange(Name, { MakeSchema(Member), SaveRange(Struct + Member.Offset, Member, Ctx) });
+}
+
+static void SaveMember(FMemberBuilder& Out, const uint8* Struct, FMemberId Name, const FSaveContext& Ctx, FStructMemberBinding Member)
+{
+	Out.AddStruct(Name, Member.Id, SaveStruct(Struct + Member.Offset, Member.Id, Ctx));
+}
+
+TUniquePtr<FBuiltStruct> SaveStruct(const uint8* Struct, FStructSchemaId Id, const FSaveContext& Ctx)
+{
+	FStructBinding Binding = Ctx.Bindings.Get(Id);
+	const FStructDeclaration& Declaration = Ctx.Declarations.Get(Id);
+	FMemberBuilder Out;
+
+	if (Binding.IsCustom())
+	{
+		Binding.AsCustom().SaveStruct(Out, Struct, nullptr, Ctx.Debug);
+	}
+	else for (FMemberVisitor It(Binding.AsSchema()); It.HasMore(); )
+	{
+		FMemberId Name = Declaration.GetMemberOrder()[It.GetIndex()];
+		switch (It.PeekKind())
+		{
+			case EMemberKind::Leaf:		SaveMember(Out, Struct, Name, Ctx, It.GrabLeaf());		break;
+			case EMemberKind::Range:	SaveMember(Out, Struct, Name, Ctx, It.GrabRange());		break;
+			case EMemberKind::Struct:	SaveMember(Out, Struct, Name, Ctx, It.GrabStruct());	break;
+		}
+	}
+	
+	return Out.BuildAndReset(Declaration, Ctx.Debug);
+}
+
+} // namespace PlainProps
\ No newline at end of file
diff --git a/Engine/Source/Programs/PlainPropsTests/Private/PlainPropsWrite.cpp b/Engine/Source/Programs/PlainPropsTests/Private/PlainPropsWrite.cpp
new file mode 100644
--- /dev/null
+++ b/Engine/Source/Programs/PlainPropsTests/Private/PlainPropsWrite.cpp
@@ -0,0 +1,842 @@
+// Copyright Epic Games, Inc. All Rights Reserved.
+
+#include  "PlainPropsWrite.h"
+#include  "PlainPropsBuildSchema.h"
+#include  "PlainPropsIndex.h"
+#include  "PlainPropsInternalBuild.h"
+#include  "PlainPropsInternalFormat.h"
+#include "Serialization/VarInt.h"
+#include "Containers/ArrayView.h"
+#include "Containers/Map.h"
+#include "Misc/ScopeExit.h"
+
+namespace PlainProps
+{
+
+// Maps declared / built ids to write ids
+// 
+// Rewrite as a more compact data structure once we get a large number of ids
+struct FWriteIds
+{
+	FWriteIds(const FIdIndexerBase& DeclaredIds, const FBuiltSchemas& Schemas);
+
+	TArray<FOptionalNameId>				Names;
+	TArray<FOptionalNestedScopeId>		NestedScopes;
+	TArray<FOptionalParametricTypeId>	ParametricTypes;
+	TArray<FOptionalSchemaId>			Structs;
+	TArray<FOptionalSchemaId>			Enums;
+
+	uint32								NumKeptSchemas;
+	TArray<FNestedScope>				KeptScopes;
+	TArray<FParametricType>				KeptParametrics;
+	TArray<FTypeId>						KeptParameters;
+
+	FNameId					Remap(FNameId Old) const				{ return Names[Old.Idx].Get(); }
+	FMemberId				Remap(FMemberId Old) const				{ return { Remap(Old.Id) }; }
+	FFlatScopeId			Remap(FFlatScopeId Old) const			{ return { Remap(Old.Name) }; }
+	FNestedScopeId			Remap(FNestedScopeId Old) const			{ return NestedScopes[Old.Idx].Get(); }
+	FScopeId				Remap(FScopeId Old) const				{ return Old.IsNested() ? FScopeId(Remap(Old.AsNested())) : FScopeId(Remap(Old.AsFlat())); }
+	FConcreteTypenameId		Remap(FConcreteTypenameId Old) const	{ return { Remap(Old.Id) }; }	
+	FParametricTypeId		Remap(FParametricTypeId Old) const		{ return ParametricTypes[Old.Idx].Get(); }
+	FTypenameId				Remap(FTypenameId Old) const			{ return Old.IsConcrete() ? FTypenameId(Remap(Old.AsConcrete())) : FTypenameId(Remap(Old.AsParametric())); }
+	FTypeId					Remap(FTypeId Old) const				{ return { Remap(Old.Scope), Remap(Old.Name) }; }
+	FSchemaId				RemapStruct(FSchemaId Old) const		{ return Structs[Old.Idx].Get(); }
+	FSchemaId				RemapEnum(FSchemaId Old) const			{ return Enums[Old.Idx].Get(); }
+};
+
+static TConstArrayView<FNameId> GetUsedNames(const FBuiltStructSchema& Used)
+{
+	static_assert(sizeof(FMemberId) == sizeof(FNameId));
+	return reinterpret_cast<const TArray<FNameId>&>(Used.MemberNames);
+}
+
+static TConstArrayView<FNameId> GetUsedNames(const FBuiltEnumSchema& Used)
+{
+	return Used.Names;
+}
+
+struct FUsedIds
+{
+	const FIdIndexerBase& Declared;
+	TBitArray<> Names;
+	TBitArray<> NestedScopes;
+	TBitArray<> ParametricTypes;
+
+
+	explicit FUsedIds(const FIdIndexerBase& InDeclared)
+	: Declared(InDeclared)
+	, Names(false, Declared.NumNames())
+	, NestedScopes(false, Declared.GetNestedScopes().Num())
+	, ParametricTypes(false, Declared.GetParametricTypes().Num())
+	{}
+
+	template<class PartialSchemaType>
+	void DetectUsage(const TArray<PartialSchemaType>& Schemas)
+	{
+		for (const PartialSchemaType& Schema : Schemas)
+		{
+			if (Schema.bUsed)
+			{
+				MarkUsed(Schema.Type);
+			}
+
+			for (FNameId Name : GetUsedNames(Schema))
+			{
+				MarkUsed(Name);
+			}
+		}
+	}
+
+	void MarkUsed(FNameId Name)
+	{
+		Names[Name.Idx] = true;
+	}
+	
+	void MarkUsed(FTypeId Type)
+	{
+		MarkUsed(Type.Scope);
+		MarkUsed(Type.Name);
+	}
+
+	void MarkUsed(FScopeId Scope)
+	{
+		if (Scope.IsFlat())
+		{
+			MarkUsed(Scope.AsFlat().Name);
+		}
+		else
+		{
+			FBitReference Used = NestedScopes[Scope.AsNested().Idx];
+			if (!Used)
+			{
+				Used = true;
+
+				FNestedScope Nested = Declared.Resolve(Scope.AsNested());
+				MarkUsed(Nested.Outer);
+				MarkUsed(Nested.Inner.Name);
+			}
+		}
+	}
+
+	void MarkUsed(FTypenameId Typename)
+	{
+		if (Typename.IsConcrete())
+		{
+			MarkUsed(Typename.AsConcrete().Id);
+		}
+		else
+		{
+			FBitReference Used = ParametricTypes[Typename.AsParametric().Idx];
+			if (!Used)
+			{
+				Used = true;
+
+				FParametricTypeView ParametricType = Declared.Resolve(Typename.AsParametric());
+				MarkUsed(ParametricType.Name.Id);
+				for (FTypeId Parameter : ParametricType.GetParameters())
+				{
+					MarkUsed(Parameter);
+				}
+			}
+		}
+	}
+};
+
+template<typename IdType>
+static uint32 MakeRemapping(TArray<TOptionalId<IdType>>& Out, const TBitArray<>& Used)
+{
+	uint32 NewIdx = 0;
+	Out.Reserve(Used.Num());
+	for (bool bUse : Used)
+	{
+		Out.Add(bUse ? ToOptional(FromIdx<IdType>(NewIdx++)) : NoId);	
+	}
+	return NewIdx;
+}
+
+static uint32 MakeRemapping(TArray<FOptionalParametricTypeId>& Out, const TBitArray<>& Used, const FParametricTypeIndexer& Declared)
+{
+	uint32 NewIdx = 0;
+	Out.Reserve(Used.Num());
+	for (bool bUse : Used)
+	{
+		uint8 NumParameters = static_cast<uint8>(Declared.At(Out.Num()).Parameters.NumParameters);
+		Out.Add(bUse ? ToOptional(FParametricTypeId(NumParameters, NewIdx++)) : NoId);	
+	}
+	return NewIdx;
+}
+template<typename T, typename C>
+static void CopyUsedIds(TArray<T>& Out, const TBitArray<>& Used, C&& Ids)
+{
+	uint32 Idx = 0;
+	for (T Id : Ids)
+	{
+		if (Used[Idx++])
+		{
+			Out.Add(Id);
+		}
+	}
+}
+
+FWriteIds::FWriteIds(const FIdIndexerBase& Declared, const FBuiltSchemas& Schemas)
+{
+	Structs.Reserve(Schemas.Structs.Num());
+	Enums.Reserve(Schemas.Enums.Num());
+
+	// Generate new struct and enum schema indices
+	uint32 NewSchemaIdx = 0;
+	for (const FBuiltStructSchema& Struct : Schemas.Structs)
+	{
+		Structs.Add(Struct.bUsed ? ToOptional(FSchemaId{ NewSchemaIdx++ }) : NoId);
+	}
+	for (const FBuiltEnumSchema& Enum : Schemas.Enums)
+	{
+		Enums.Add(Enum.bUsed ? ToOptional(FSchemaId{ NewSchemaIdx++ }) : NoId);
+	}
+	NumKeptSchemas = NewSchemaIdx;
+
+	// Generate new name indices
+	FUsedIds Used(Declared);
+	Used.DetectUsage(Schemas.Structs);
+	Used.DetectUsage(Schemas.Enums);
+	MakeRemapping(Names, Used.Names);
+
+	// Generate new nested scope and parametric type mappings
+	KeptScopes.Reserve(MakeRemapping(/* out */ NestedScopes, Used.NestedScopes));
+	KeptParametrics.Reserve(MakeRemapping(/* out */ ParametricTypes, Used.ParametricTypes, Declared.GetParametricTypes()));
+
+	// Copy nested scopes and parametric types
+	CopyUsedIds(/* out */ KeptScopes, Used.NestedScopes, Declared.GetNestedScopes());
+	CopyUsedIds(/* out */ KeptParametrics, Used.ParametricTypes, Declared.GetParametricTypes().GetAllTypes());
+
+	// Remap copied nested scopes
+	for (FNestedScope& KeptScope : KeptScopes)
+	{
+		KeptScope.Inner = Remap(KeptScope.Inner);
+		KeptScope.Outer = Remap(KeptScope.Outer);
+	}
+
+	// Remap copied parametric types and copy parameters
+	TConstArrayView<FTypeId> DeclaredParameters = Declared.GetParametricTypes().GetAllParameters();
+	for (FParametricType& KeptType : KeptParametrics)
+	{
+		FParameterIndexRange OldParameters = KeptType.Parameters;
+		KeptType.Name = Remap(KeptType.Name);
+		KeptType.Parameters.Idx = KeptParameters.Num();
+		KeptParameters.Append(DeclaredParameters.Slice(OldParameters.Idx, OldParameters.NumParameters));
+	}
+
+	// Remap copied parameters
+	for (FTypeId& KeptParameter : KeptParameters)
+	{
+		KeptParameter = Remap(KeptParameter);
+	}
+}
+
+//////////////////////////////////////////////////////////////////////////
+
+static TArray<FMemberType> GetMemberTypes(const FBuiltStructSchema& Struct)
+{
+	TArray<FMemberType> Out;
+	Out.Reserve(Struct.MemberSchemas.Num());
+	for (const FMemberSchema* Schema : Struct.MemberSchemas)
+	{
+		Out.Add(Schema->Type);
+	}
+	return Out;
+}
+
+static TArray<FMemberType> GetInnerRangeTypes(const FBuiltStructSchema& Struct)
+{
+	TArray<FMemberType> Out;
+	for (const FMemberSchema* Schema : Struct.MemberSchemas)
+	{
+		Out.Append(Schema->InnerRangeTypes);
+	}
+	return Out;
+}
+
+static FOptionalSchemaId GetStaticInnerSchema(const FMemberSchema& Schema, const FWriteIds& NewIds)
+{
+	if (FOptionalSchemaId InnerSchema = Schema.InnerSchema)
+	{
+		FMemberType InnermostType = Schema.GetInnermostType();
+		check(IsStructOrEnum(InnermostType));
+		if (InnermostType.IsLeaf())
+		{
+			return ToOptional(NewIds.RemapEnum(InnerSchema.Get()));
+		}
+		else if (!InnermostType.AsStruct().IsDynamic)
+		{
+			return ToOptional(NewIds.RemapStruct(InnerSchema.Get()));
+		}
+	}
+
+	return {};
+}
+
+static TArray<FSchemaId> GetInnerSchemas(const FBuiltStructSchema& Struct, const FWriteIds& NewIds, ESuper Inheritance)
+{
+	TArray<FSchemaId> Out;	
+	if (Inheritance != ESuper::No && Inheritance != ESuper::Reused)
+	{
+		Out.Add(NewIds.RemapStruct(Struct.Super.Get()));
+	}
+
+	for (const FMemberSchema* Schema : Struct.MemberSchemas)
+	{
+		if (FOptionalSchemaId InnerSchema = GetStaticInnerSchema(*Schema, NewIds))
+		{
+			Out.Add(InnerSchema.Get());
+		}
+	}
+	return Out;
+}
+
+template<typename IdType>
+static TArray<IdType> RemapIds(const FWriteIds& NewIds, const TArray<IdType>& Names)
+{
+	TArray<IdType> Out = Names;
+	for (IdType& Name : Out)
+	{
+		Name = NewIds.Remap(Name);
+	}
+	return Out;
+}
+
+static ESuper GetInheritance(FOptionalSchemaId Super, TConstArrayView<const FMemberSchema*> Members)
+{
+	if (Super)
+	{
+		if (Members.IsEmpty() || !IsSuper(Members[0]->Type))
+		{
+			return ESuper::Unused;
+		}
+		
+		return Members[0]->Type.AsStruct().IsDynamic || Members[0]->InnerSchema != Super ? ESuper::Used : ESuper::Reused;
+	}
+
+	return ESuper::No;
+}
+
+static void WriteSchema(TArray64<uint8>& Out, const FBuiltStructSchema& Struct, const FWriteIds& NewIds)
+{
+	FOptionalSchemaId NewSuperId = Struct.Super ? ToOptional(NewIds.RemapStruct(Struct.Super.Get())) : NoId;
+	ESuper Inheritance = GetInheritance(NewSuperId, Struct.MemberSchemas);
+
+	TArray<FMemberType> MemberTypes = GetMemberTypes(Struct);
+	TArray<FMemberType> InnerRangeTypes = GetInnerRangeTypes(Struct);
+	TArray<FMemberId> MemberNames = RemapIds(NewIds, Struct.MemberNames);
+	TArray<FSchemaId> InnerSchemas = GetInnerSchemas(Struct, NewIds, Inheritance);
+
+	check(MemberNames.Num() + UsesSuper(Inheritance) == MemberTypes.Num());
+
+	// Zero init for determinism
+	alignas(FStructSchema) uint8 HeaderData[offsetof(FStructSchema, Footer)] = {};
+	FStructSchema& BinaryHeader = *reinterpret_cast<FStructSchema*>(HeaderData);
+	BinaryHeader.Type = NewIds.Remap(Struct.Type);
+	BinaryHeader.Inheritance = Inheritance;
+	BinaryHeader.IsDense = Struct.bDense;
+	BinaryHeader.NumMembers = IntCastChecked<uint16>(MemberTypes.Num());
+	BinaryHeader.NumRangeTypes = IntCastChecked<uint16>(InnerRangeTypes.Num());
+	BinaryHeader.NumInnerSchemas = IntCastChecked<uint16>(InnerSchemas.Num());
+
+	check(IsAligned(Out.Num(), alignof(FStructSchema)));
+	WriteData(Out, HeaderData, sizeof(HeaderData));
+	WriteArray(Out, MemberTypes);
+	WriteArray(Out, InnerRangeTypes);
+	WriteAlignedArray(Out, MakeArrayView(MemberNames));
+	WriteAlignedArray(Out, MakeArrayView(InnerSchemas));
+}
+
+static bool IsFlatSequence(TConstArrayView<uint64> Constants)
+{
+	uint64 Expected = 0;
+	for (uint64 Constant : Constants)
+	{
+		if (Constant != Expected)
+		{
+			return false;
+		}
+		++Expected;
+	}
+	return true;
+}
+static bool IsFlagSequence(TConstArrayView<uint64> Constants)
+{
+	check(Constants.Num() <= 64);
+	uint64 Expected = 1;
+	for (uint64 Constant : Constants)
+	{
+		if (Constant != Expected)
+		{
+			return false;
+		}
+		Expected <<= 1;
+	}
+	return true;
+}
+
+template<typename IntType>
+static void WriteEnumConstantsAs(TArray64<uint8>& Out, TConstArrayView<uint64> Constants)
+{
+	TArray<IntType, TInlineAllocator<64>> Tmp;
+	for (uint64 Constant : Constants)
+	{
+		Tmp.Add(IntCastChecked<IntType>(Constant));
+	}
+	WriteArray(Out, Constants);
+}
+
+static void WriteEnumConstants(TArray64<uint8>& Out, ELeafWidth Width, TConstArrayView<uint64> Constants)
+{
+	switch (Width)
+	{
+		case ELeafWidth::B8:	WriteEnumConstantsAs<uint8 >(Out, Constants); break;
+		case ELeafWidth::B16:	WriteEnumConstantsAs<uint16>(Out, Constants); break;
+		case ELeafWidth::B32:	WriteEnumConstantsAs<uint32>(Out, Constants); break;
+		case ELeafWidth::B64:	WriteArray(Out, Constants); break;
+	}
+}
+
+static void WriteSchema(TArray64<uint8>& Out, const FBuiltEnumSchema& Enum, const FWriteIds& NewIds)
+{
+	bool bIsSequence = Enum.Mode == EEnumMode::Flag ? IsFlagSequence(Enum.Constants) : IsFlatSequence(Enum.Constants);
+	TArray<FNameId> Names = RemapIds(NewIds, Enum.Names);
+
+	FEnumSchema BinaryHeader = { NewIds.Remap(Enum.Type) };
+	BinaryHeader.FlagMode = Enum.Mode == EEnumMode::Flag;
+	BinaryHeader.ExplicitConstants = !bIsSequence;
+	BinaryHeader.Width = Enum.Width;
+	BinaryHeader.Num = IntCastChecked<uint16>(Names.Num());
+
+	check(IsAligned(Out.Num(), alignof(FEnumSchema)));
+	WriteData(Out, &BinaryHeader, sizeof(BinaryHeader));
+	WriteArray(Out, Names);
+	if (BinaryHeader.ExplicitConstants)
+	{
+		WriteEnumConstants(Out, Enum.Width, Enum.Constants);
+	}
+	WriteAlignmentPadding<FEnumSchema>(Out);
+}
+
+template<typename T>
+void AppendBinary(TArray64<uint8>& Dst, const TArray<T>& Src)
+{
+	if (uint32 NumBytes = Src.Num() * sizeof(T))
+	{
+		Dst.AddUninitialized(NumBytes);
+		FMemory::Memcpy(Dst.GetData() + Dst.Num() - NumBytes, Src.GetData(), NumBytes);
+	}
+}
+
+static void WriteSchemasImpl(TArray64<uint8>& Out, const FBuiltSchemas& Schemas, const FWriteIds& NewIds)
+{
+	check(NewIds.Structs.Num() == Schemas.Structs.Num());
+	check(NewIds.Enums.Num() == Schemas.Enums.Num());
+
+	// @see FReadSchemaBatch
+
+	WriteAlignmentPadding<uint32>(Out);
+	const int64 HeaderPos = Out.Num();
+	const uint32 NumSchemas = NewIds.NumKeptSchemas;
+	Out.AddUninitialized(sizeof(FSchemaBatch) + NumSchemas * sizeof(uint32));
+
+	TArray<uint32> SchemaOffsets;
+	SchemaOffsets.Reserve(NumSchemas);
+	auto WritePartialSchemas = [&Out, &SchemaOffsets, HeaderPos, &NewIds](auto& PartialSchemas)
+	{
+		for (auto& PartialSchema : PartialSchemas)
+		{
+			if (PartialSchema.bUsed)
+			{
+				SchemaOffsets.Add(IntCastChecked<uint32>(Out.Num() - HeaderPos));
+				WriteSchema(Out, PartialSchema, NewIds);
+			}
+		}
+	};
+
+	WritePartialSchemas(Schemas.Structs);
+	WritePartialSchemas(Schemas.Enums);
+	check(SchemaOffsets.Num() == NumSchemas);
+
+	// Write header 
+	const int64 NestedScopePos = Out.Num();
+	FSchemaBatch Header{0};
+	Header.NumNestedScopes = NewIds.KeptScopes.Num();
+	Header.NestedScopesOffset = IntCastChecked<uint32>(NestedScopePos - HeaderPos);
+	Header.NumParametricTypes = NewIds.KeptParametrics.Num();
+	Header.NumSchemas = NewIds.NumKeptSchemas;
+	FMemory::Memcpy(&Out[HeaderPos], &Header, sizeof(FSchemaBatch));
+
+	// Write schema offsets
+	FMemory::Memcpy(&Out[HeaderPos] + sizeof(FSchemaBatch), SchemaOffsets.GetData(), NumSchemas * sizeof(uint32));
+
+	AppendBinary(Out, NewIds.KeptScopes);
+	AppendBinary(Out, NewIds.KeptParametrics);
+	AppendBinary(Out, NewIds.KeptParameters);
+}
+
+//////////////////////////////////////////////////////////////////////////
+
+uint64 WriteSkippableSlice(TArray64<uint8>& Out, TConstArrayView64<uint8> Slice)
+{
+	if (uint64 Num = IntCastChecked<uint64>(Slice.Num()))
+	{
+		uint32 VarIntBytes = MeasureVarUInt(Num);
+		int64 VarIntPos = Out.AddUninitialized(VarIntBytes + Slice.Num());
+		WriteVarUInt(Num, &Out[VarIntPos]);
+		FMemory::Memcpy(&Out[VarIntPos + VarIntBytes], Slice.GetData(), Slice.Num());
+		return VarIntBytes + Num;
+	}
+
+	constexpr uint8 ZeroVarUInt = 0; 
+	Out.Add(ZeroVarUInt);
+	return 1;
+}
+
+//////////////////////////////////////////////////////////////////////////
+
+class FBitCacheWriter
+{
+	static constexpr int64 Unused = -1;
+	static constexpr int64 Finished = -2;
+
+	uint8 Bits = 0;
+	uint32 NumLeft = 0;
+	TArray64<uint8>& Dest;
+	int64 DestIdx = Unused;
+
+	void DoFlush()
+	{
+		if (DestIdx >= 0)
+		{
+			Dest[DestIdx] = Bits;
+		}
+	}
+
+public:
+	explicit FBitCacheWriter(TArray64<uint8>& Out) : Dest(Out) {}
+	~FBitCacheWriter()
+	{
+		check(DestIdx == Finished);
+	}
+
+	void WriteBit(bool Bit)
+	{
+		if (NumLeft == 0)
+		{
+			check(DestIdx != Finished);
+			DoFlush();
+			
+			DestIdx = Dest.Num();
+			Dest.Add(0);
+
+			Bits = Bit;
+			NumLeft = 7;
+		}
+		else
+		{
+			Bits |= uint8(Bit) << (8 - NumLeft);
+			--NumLeft;
+		}
+	}
+
+	void Flush()
+	{
+		DoFlush();
+		NumLeft = 0;
+		DestIdx = Finished;
+	}
+};
+
+//////////////////////////////////////////////////////////////////////////
+
+class FMemberWriter
+{
+public:
+	FMemberWriter(TArray64<uint8>& Out, const FBuiltSchemas& InSchemas, const FWriteIds& InNewIds, const FDebugIds& InDebug)
+	: Bytes(Out)
+	, Bits(Out)
+	, Schemas(InSchemas)
+	, NewIds(InNewIds)
+	, Debug(InDebug)
+	{}
+
+	void WriteMembers(FStructSchemaId BuiltId, const FBuiltStruct& Struct)
+	{
+		const FBuiltStructSchema& Schema = Schemas.Structs[BuiltId.Idx];
+		const TArray<FMemberId>& Order = Schema.MemberNames;
+		const int32 NumSuper = Schema.MemberSchemas.Num() - Schema.MemberNames.Num();
+		check(NumSuper == 0 || NumSuper == 1 && IsSuper(Schema.MemberSchemas[0]->Type));
+		check(Struct.NumMembers <= Schema.MemberSchemas.Num());
+
+		int32 Idx = 0;
+		if (Schema.bDense)
+		{
+			for (const FBuiltMember& Member : MakeArrayView(Struct.Members, Struct.NumMembers))
+			{
+				checkf(Member.Name == NoId || Order[Idx - NumSuper] == Member.Name.Get(), TEXT("Member '%s' in '%s' %s %s"),
+					*Debug.Print(Member.Name), *Debug.Print(Schema.Type),
+					Order.Contains(Member.Name) ? TEXT("appeared before missing member") : TEXT("is undeclared"),
+					Order.Contains(Member.Name) ? *Debug.Print(Order[Idx]) : TEXT(""));
+
+				WriteMember(Schema.MemberSchemas[Idx]->GetInnermostType(), Member.Schema, Member.Value);
+				++Idx;
+			}
+		}
+		else
+		{
+			for (const FBuiltMember& Member : MakeArrayView(Struct.Members, Struct.NumMembers))
+			{
+				for (bool bSkip = true; bSkip; ++Idx) 
+				{
+					checkf(Idx - NumSuper < Order.Num(), TEXT("Member '%s' in '%s' %s"), *Debug.Print(Member.Name), *Debug.Print(Schema.Type), 
+						   Order.Contains(Member.Name) ? TEXT("appeared in non-declared order") : TEXT("is undeclared"));
+
+					bSkip = Member.Name && (Idx < NumSuper || Order[Idx - NumSuper] != Member.Name.Get());
+					Bits.WriteBit(bSkip);
+				}
+
+				WriteMember(Schema.MemberSchemas[Idx - 1]->GetInnermostType(), Member.Schema, Member.Value);
+			}
+
+			// Skip remaining missing members
+			for (int32 Num = Schema.MemberSchemas.Num(); Idx < Num; ++Idx)
+			{
+				Bits.WriteBit(true);
+			}
+		}
+
+		Bits.Flush();
+	}
+
+private:
+	TArray64<uint8>& Bytes;
+	FBitCacheWriter Bits;
+	TArray64<uint8> Tmp;
+	const FBuiltSchemas& Schemas;
+	const FWriteIds& NewIds;
+	const FDebugIds& Debug;
+
+	// Tricky! InnermostType comes from FBuiltStructSchema and it's IsDynamic is decided during noting.
+	// Schema.GetInnermostType() IsDynamic is false, it's from immutable FBuiltMember built before noting.
+	void WriteMember(FMemberType InnermostType, const FMemberSchema& Schema, const FBuiltValue& Value)
+	{
+		if (InnermostType.IsStruct() && InnermostType.AsStruct().IsDynamic)
+		{
+			FSchemaId NewId = NewIds.RemapStruct(Schema.InnerSchema.Get());
+			Bytes.Append(reinterpret_cast<uint8*>(&NewId.Idx), sizeof(NewId.Idx));
+		}
+
+		switch (Schema.Type.GetKind())
+		{
+		case EMemberKind::Leaf:		WriteLeaf(Schema.Type.AsLeaf(), Value.Leaf); break;
+		case EMemberKind::Struct:	WriteStruct(Schema.Type.AsStruct(), static_cast<FStructSchemaId>(Schema.InnerSchema.Get()), *Value.Struct); break;
+		case EMemberKind::Range:	WriteRange(Schema.Type.AsRange().MaxSize, Schema.InnerRangeTypes, Schema.InnerSchema, Value.Range); break;
+		}
+	}
+
+	void WriteLeaf(FLeafType LeafType, uint64 LeafValue)
+	{
+		if (LeafType.Type == ELeafType::Bool)
+		{
+			check(LeafValue <= 1);
+			Bits.WriteBit(!!LeafValue);
+		}
+		else
+		{
+			WriteUnsigned(LeafValue, SizeOf(LeafType.Width));
+		}
+	}
+		
+	void WriteStruct(FStructType StructType, FStructSchemaId Schema, const FBuiltStruct& Struct)
+	{
+		Tmp.Reserve(1024);
+		FMemberWriter(Tmp, Schemas, NewIds, Debug).WriteMembers(Schema, Struct);
+		WriteSkippableSlice(Bytes, Tmp);
+		Tmp.Reset();
+	}
+
+	void WriteRange(ERangeSizeType NumType, TConstArrayView<FMemberType> Types, FOptionalSchemaId InnermostSchema, const FBuiltRange* Range)
+	{
+		check(Types.Num() > 0);
+		check((Types.Num() > 1) == (Types[0].GetKind() == EMemberKind::Range));
+		check(!Range || Range->Num > 0 && Range->Num <= Max(NumType));
+
+		// Write Num
+		uint64 Num = Range ? Range->Num : 0;
+		if (NumType == ERangeSizeType::Uni)
+		{
+			Bits.WriteBit(Num == 1);
+		}
+		else
+		{
+			WriteUnsigned(Num, SizeOf(NumType));
+		}
+		
+		// Write Data
+		if (Range)
+		{
+			switch (Types[0].GetKind())
+			{
+			case EMemberKind::Leaf:		WriteLeaves(Types[0].AsLeaf(), *Range);	break;
+			case EMemberKind::Struct:	WriteStructs(Types[0].AsStruct(), static_cast<FStructSchemaId>(InnermostSchema.Get()), Range->AsStructs()); break;
+			case EMemberKind::Range:	WriteRanges(Types[0].AsRange().MaxSize, Types.RightChop(1), InnermostSchema, Range->AsRanges()); break;
+			}
+		}
+	}
+
+	void WriteLeaves(FLeafType Leaf, const FBuiltRange& Range)
+	{
+		if (Leaf.Type == ELeafType::Bool)
+		{
+			FBitCacheWriter BitArray(/* Out */ Bytes);
+			for (uint8 Bool : TConstArrayView64<uint8>(Range.Data, Range.Num))
+			{
+				check(Bool <= 1);
+				BitArray.WriteBit(!!Bool);
+			}
+			BitArray.Flush();
+		}
+		else
+		{
+			WriteData(Bytes, Range.Data, GetLeafRangeSize(Range.Num, Leaf));
+		}
+	}
+
+	template<typename ItemType, typename WriteFn>
+	void WriteSkippableItems(TConstArrayView64<ItemType> Items, WriteFn WriteItem)
+	{
+		Tmp.Reserve(1024);
+		FMemberWriter NestedWriter(Tmp, Schemas, NewIds, Debug);
+		for (const ItemType& Item : Items)
+		{
+			WriteItem(NestedWriter, Item);
+		}
+		NestedWriter.Bits.Flush();
+
+		WriteSkippableSlice(/* out */ Bytes, Tmp);
+		Tmp.Reset();
+	}
+
+	void WriteStructs(FStructType StructType, FStructSchemaId Schema, TConstArrayView64<TUniquePtr<const FBuiltStruct>> Structs)
+	{
+		Bytes.Append(reinterpret_cast<uint8*>(&Schema.Idx), StructType.IsDynamic * sizeof(Schema.Idx));
+		WriteSkippableItems(Structs, [=](FMemberWriter& Out, const TUniquePtr<const FBuiltStruct>& Struct) { Out.WriteStruct(StructType, Schema, *Struct); });
+	}
+	
+	void WriteRanges(ERangeSizeType NumType, TConstArrayView<FMemberType> Types, FOptionalSchemaId InnermostSchema, TConstArrayView64<const FBuiltRange*> Ranges)
+	{	
+		WriteSkippableItems(Ranges, [=](FMemberWriter& Out, const FBuiltRange* Range) { Out.WriteRange(NumType, Types, InnermostSchema, Range); });
+	}
+
+	void WriteUnsigned(uint64 Value, SIZE_T SizeOf)
+	{
+		check(SizeOf == 8 || (Value >> (SizeOf*8)) == 0);
+		Bytes.Append(reinterpret_cast<uint8*>(&Value), SizeOf);
+	}
+};
+
+////////////////////////////////////////////////////////////////////////////////////////////////////////////////////////////////////////////////////
+
+FWriter::FWriter(const FIdIndexerBase& DeclaredIds, const FBuiltSchemas& InSchemas)
+: Schemas(InSchemas)
+, Debug(DeclaredIds)
+, NewIds(new FWriteIds(DeclaredIds, InSchemas))
+{}
+
+FWriter::~FWriter()
+{}
+
+bool FWriter::Uses(FNameId BuiltId) const
+{
+	return !!NewIds->Names[BuiltId.Idx];
+}
+
+FOptionalStructSchemaId FWriter::GetWriteId(FStructSchemaId BuiltId) const
+{
+	return static_cast<FOptionalStructSchemaId>(NewIds->Structs[BuiltId.Idx]);
+}
+
+void FWriter::WriteSchemas(TArray64<uint8>& Out) const
+{
+	WriteSchemasImpl(Out, Schemas, *NewIds);
+}
+
+void FWriter::WriteMembers(TArray64<uint8>& Out, FStructSchemaId BuiltId, const FBuiltStruct& Struct) const
+{
+	FMemberWriter(Out, Schemas, *NewIds, Debug).WriteMembers(BuiltId, Struct);
+}
+
+////////////////////////////////////////////////////////////////////////////////////////////////////////////////////////////////////////////////////
+
+void FDebugIds::AppendDebugString(FString& Out, FScopeId Scope) const
+{
+	if (Scope.IsFlat())
+	{
+		AppendDebugString(Out, Scope.AsFlat().Name);
+	}
+	else
+	{
+		FNestedScope Nested = Resolve(Scope.AsNested());
+		AppendDebugString(Out, Nested.Outer);
+		Out.AppendChar('.');
+		AppendDebugString(Out, Nested.Inner.Name);
+	}
+}
+
+void FDebugIds::AppendDebugString(FString& Out, FTypenameId Typename) const
+{
+	if (Typename.IsConcrete())
+	{
+		AppendDebugString(Out, Typename.AsConcrete().Id);
+	}
+	else
+	{
+		FParametricTypeView ParametricType = Resolve(Typename.AsParametric());
+		AppendDebugString(Out, ParametricType.Name.Id);
+		Out.AppendChar('<');
+		for (FTypeId Parameter : ParametricType.GetParameters())
+		{
+			AppendDebugString(Out, Parameter);
+			Out.AppendChar(',');
+		}
+		Out.GetCharArray()[Out.Len() - 1] = '>';
+	}
+}
+
+void FDebugIds::AppendDebugString(FString& Out, FTypeId Type) const
+{
+	AppendDebugString(Out, Type.Scope);
+	Out.AppendChar('.');
+	AppendDebugString(Out, Type.Name);
+}
+
+FString FDebugIds::Print(FNameId Name) const
+{
+	FString Out;
+	AppendDebugString(Out, Name);
+	return Out;
+}
+
+FString FDebugIds::Print(FMemberId Name) const
+{
+	FString Out;
+	AppendDebugString(Out, Name.Id);
+	return Out;
+}
+
+FString FDebugIds::Print(FOptionalMemberId Name) const
+{
+	return Name ? Print(Name.Get().Id) : TEXT("!super!");
+}
+
+FString FDebugIds::Print(FTypeId Type) const
+{
+	FString Out;
+	AppendDebugString(Out, Type);
+	return Out;
+}
+
+} // namespace PlainProps
diff --git a/Engine/Source/Programs/PlainPropsTests/Public/PlainPropsBind.h b/Engine/Source/Programs/PlainPropsTests/Public/PlainPropsBind.h
new file mode 100644
--- /dev/null
+++ b/Engine/Source/Programs/PlainPropsTests/Public/PlainPropsBind.h
@@ -0,0 +1,702 @@
+// Copyright Epic Games, Inc. All Rights Reserved.
+
+#pragma once
+
+#include "Containers/ArrayView.h"
+#include "Memory/MemoryFwd.h"
+#include  "PlainPropsTypes.h"
+#include  "PlainPropsCtti.h"
+#include  "PlainPropsDeclare.h"
+
+namespace PlainProps 
+{
+
+// Todo: Include Build.h and Load.h instead?
+enum class EMemberPresence;
+enum class EEnumMode;
+struct FBuiltStruct;
+struct FLoadBatch;
+class FMemberBuilder;
+class FStructBinding;
+struct FStructView;
+class FRangeBinding;
+struct FTypedRange;
+class IRangeBinding;
+template<class T> class TIdIndexer;
+
+//////////////////////////////////////////////////////////////////////////
+
+struct FIdBinding
+{
+	TConstArrayView<FNameId>			Names;
+	TConstArrayView<FNestedScopeId>		Scopes;
+	TConstArrayView<FParametricTypeId>	Parametrics;
+};
+
+class FIdBinder
+{
+	const uint32 NumNames;
+	const uint32 NumScopes;
+	const uint32 NumParameters;
+	const FNameId* Data;
+
+public:
+	template<class NameType>
+	FIdBinder(TIdIndexer<NameType>& Indexer, TConstArrayView<NameType> Names, TConstArrayView<FNestedScope> Scopes, TConstArrayView<FParametricType> ParametericTypes);
+
+	FIdBinding Get() const;
+};
+
+
+
+////////////////////////////////////////////////////////////////////////////////////////////////
+
+enum class ELeafBindType : uint8 { Bool, IntS, IntU, Float, Hex, Enum, Unicode, BitfieldBool };
+
+inline static constexpr ELeafBindType ToLeafBindType(ELeafType Type)
+{
+	return static_cast<ELeafBindType>(static_cast<uint8>(Type));
+}
+
+inline static constexpr ELeafType ToLeafType(ELeafBindType Type)
+{
+	return Type == ELeafBindType::BitfieldBool ? ELeafType::Bool : static_cast<ELeafType>(static_cast<uint8>(Type));
+}
+
+struct FArithmeticBindType
+{
+    EMemberKind		_ : 2;
+	ELeafBindType	__ : 3;
+	ELeafWidth		Width : 2;
+};
+
+struct FBitfieldBindType
+{
+    EMemberKind		_ : 2;
+	ELeafBindType	__ : 3;
+	uint8			Idx : 3;
+};
+
+union FLeafBindType
+{
+	constexpr explicit FLeafBindType(FLeafType In) : Arithmetic({EMemberKind::Leaf, ToLeafBindType(In.Type), In.Width}) {}
+	constexpr explicit FLeafBindType(FBitfieldBindType In) : Bitfield({EMemberKind::Leaf, ELeafBindType::BitfieldBool, In.Idx}) {}
+
+	struct
+	{
+		EMemberKind			_ : 2;
+		ELeafBindType		Type : 3;
+	}						Bind;
+	FArithmeticBindType		Arithmetic;
+	FBitfieldBindType		Bitfield;
+};
+
+inline static constexpr FLeafType ToLeafType(FLeafBindType Leaf)
+{
+	if (Leaf.Bind.Type == ELeafBindType::BitfieldBool)
+	{
+		return { EMemberKind::Leaf, ELeafWidth::B8, ELeafType::Bool };
+	}
+	
+	return { EMemberKind::Leaf,  Leaf.Arithmetic.Width, ToLeafType(Leaf.Bind.Type) };
+}
+
+struct FRangeBindType : FRangeType {};
+
+struct FStructBindType : FStructType {};
+
+union FMemberBindType
+{
+	constexpr explicit FMemberBindType(FLeafType In) : Leaf(In) {}
+	constexpr explicit FMemberBindType(FBitfieldBindType In) : Leaf(In) {}
+	constexpr explicit FMemberBindType(FRangeType In) : Range(In) {}
+	constexpr explicit FMemberBindType(ERangeSizeType MaxSize) : Range({EMemberKind::Range, MaxSize}) {}
+	constexpr explicit FMemberBindType(FStructType In) : Struct(In) {}
+	
+	bool				IsLeaf() const			{ return Kind == EMemberKind::Leaf; }
+	bool				IsRange() const			{ return Kind == EMemberKind::Range; }
+	bool				IsStruct() const		{ return Kind == EMemberKind::Struct; }
+	EMemberKind			GetKind() const			{ return Kind; }
+	
+	FLeafBindType		AsLeaf() const			{ check(IsLeaf());		return Leaf; }
+	FRangeBindType		AsRange() const			{ check(IsRange());		return Range; }
+	FStructBindType		AsStruct() const		{ check(IsStruct());	return Struct; }
+	uint8				AsByte() const			{ return BitCast<uint8>(*this); }
+
+	friend inline bool operator==(FMemberBindType A, FMemberBindType B) { return A.AsByte() == B.AsByte(); }
+
+private:
+    EMemberKind			Kind : 2;
+	FLeafBindType		Leaf;
+    FRangeBindType		Range;
+    FStructBindType		Struct;
+};
+
+static_assert(sizeof(FMemberBindType) == 1);
+
+////////////////////////////////////////////////////////////////////////////////////////////////
+
+// Members are loaded in saved FStructSchema order, not current offset order unless upgrade layer reorders
+struct FStructSchemaBinding
+{
+	uint16					NumMembers;
+	uint16					NumInnerSchemas;
+	uint16					NumInnerRanges;
+	FMemberBindType			Members[0];
+
+	const FMemberBindType*	GetInnerRangeTypes() const	{ return Members + NumMembers; }
+	const uint32*			GetOffsets() const			{ return AlignPtr<uint32>(GetInnerRangeTypes() + NumInnerRanges); }
+	const FSchemaId*		GetInnerSchemas() const		{ return AlignPtr<FSchemaId>(GetOffsets() + NumMembers); }
+	const FRangeBinding*	GetRangeBindings() const	{ return AlignPtr<FRangeBinding>(GetInnerSchemas() + NumInnerSchemas); }
+};
+
+struct FUnpackedLeafBindType
+{
+	ELeafBindType			Type;
+	union
+	{
+		ELeafWidth				Width;
+		uint8					BitfieldIdx;
+	};
+	
+	//constexpr FUnpackedLeafBindType(ELeafBindType InType, ELeafWidth InWidth) : Type(InType), Width(InWidth) {}
+	constexpr FUnpackedLeafBindType(FLeafBindType In)
+		: Type(In.Bind.Type)
+	{
+		if (Type == ELeafBindType::BitfieldBool)
+		{
+			Width = In.Arithmetic.Width;
+		}
+		else
+		{
+			BitfieldIdx = In.Bitfield.Idx;
+		}
+	}
+
+//	constexpr bool operator==(FUnpackedLeafBindType O) { return Type == O.Type && Width == O.Width; }
+//	FMemberBindType Pack() const { return FMemberBindType(Type, Width); }
+};
+
+struct FLeafMemberBinding
+{
+	FUnpackedLeafBindType	Leaf;
+	FOptionalEnumSchemaId	Enum;
+	uint64					Offset;
+};
+
+struct FRangeMemberBinding
+{
+	const FMemberBindType*	InnerTypes;
+	const FRangeBinding*	RangeBindings;
+	FOptionalSchemaId		InnermostSchema;
+	uint64					Offset;
+};
+
+struct FStructMemberBinding
+{
+	FStructSchemaId		Id;
+	uint64				Offset;
+};
+
+// Iterates over member bindings
+class FMemberVisitor
+{
+public:
+	explicit FMemberVisitor(const FStructSchemaBinding& InSchema);
+	bool						HasMore() const			{ return MemberIdx < NumMembers; }
+	uint16						GetIndex() const		{ return MemberIdx; }
+	
+	EMemberKind					PeekKind() const;		// @pre HasMore()
+	FMemberBindType				PeekType() const;		// @pre HasMore()
+
+	FLeafMemberBinding			GrabLeaf();				// @pre PeekKind() == EMemberKind::Leaf
+	FRangeMemberBinding			GrabRange();			// @pre PeekKind() == EMemberKind::Range
+	FStructMemberBinding		GrabStruct();			// @pre PeekKind() == EMemberKind::Struct
+
+protected: // for unit tests
+	const FStructSchemaBinding& Schema;
+	//const FMemberBindType*	Footer;
+	const uint16				NumMembers;
+	//const uint16			NumInnerRanges;			// Number of ranges and nested ranges
+	//const uint16			NumInnerSchemas;		// Number of static structs and enums
+	uint16						MemberIdx = 0;
+	uint16						InnerRangeIdx = 0;		// Types of [nested] ranges
+	uint16						InnerSchemaIdx = 0;		// Types of static structs and enums
+
+	//const uint32*				GetOffsets() const;
+	//const FMemberBindType*		GetInnerRanges() const;
+	//const FSchemaId*			GetInnerSchemas() const;
+
+	using FMemberBindTypeRange = TConstArrayView<FMemberBindType>;
+
+	uint64						GrabMemberOffset();
+	FMemberBindTypeRange		GrabInnerTypes();
+	FSchemaId					GrabInnerSchema();
+	FStructSchemaId				GrabStructSchema(FStructType Type);
+	FOptionalSchemaId			GrabRangeSchema(FMemberType InnermostType);
+	FEnumSchemaId				GrabEnumSchema()			{ return static_cast<FEnumSchemaId&&>(GrabInnerSchema()); }
+};
+
+////////////////////////////////////////////////////////////////////////////////////////////////
+
+enum ECustomLoadMethod { Construct, Assign };
+
+class ICustomStructBinding
+{
+public:
+	virtual void			LoadStruct(void* Dst, FStructView Src, ECustomLoadMethod Method, const FLoadBatch& Batch) const = 0;
+	virtual void			SaveStruct(FMemberBuilder& Dst, const void* Src, void* UserData, const FDebugIds& Debug) const = 0;
+};
+
+class FStructBinding
+{
+public:
+	explicit FStructBinding(const FStructSchemaBinding& Schema) : Handle(uint64(&Schema) | SchemaBit) {}
+	explicit FStructBinding(const ICustomStructBinding& Custom) : Handle(uint64(&Custom)) {}
+	
+	bool						IsSchema() const	{ return Handle & SchemaBit; }
+	bool						IsCustom() const	{ return !IsSchema(); }
+	const FStructSchemaBinding&	AsSchema() const	{ check(IsSchema()); return *reinterpret_cast<FStructSchemaBinding*>(Handle & ~SchemaBit); }
+	const ICustomStructBinding&	AsCustom() const	{ check(IsCustom()); return *reinterpret_cast<ICustomStructBinding*>(Handle & ~SchemaBit); }
+
+private:
+	static constexpr uint64 SchemaBit = 1;
+	uint64 Handle;
+
+	friend class FStructBindings;
+	bool						IsBound() const		{ return Handle != 0; }
+	FStructSchemaBinding*		TryGetSchema() 		{ return (IsBound() && IsSchema()) ? reinterpret_cast<FStructSchemaBinding*>(Handle & ~SchemaBit) : nullptr; }
+	FStructBinding() : Handle(0) {}
+};
+
+////////////////////////////////////////////////////////////////////////////////////////////////
+
+class FConstructionRequest
+{
+	void* const Range = nullptr;
+	const uint64 Num = 0;
+	uint64 Index = 0;
+
+	friend class FRangeLoader;
+	FConstructionRequest(void* InRange, uint64 InNum) : Range(InRange), Num(InNum) {}
+	
+public:
+	template<typename T>
+	T& GetRange() const { return *reinterpret_cast<T*>(Range); }
+	
+	uint64 NumTotal() const { return Num; }
+	uint64 NumMore() const { return Num - Index; }
+	uint64 GetIndex() const { return Index; }
+	bool IsFirstCall() const { return Index == 0;}
+	bool IsFinalCall() const { return Index == Num;}
+};
+
+class FConstructedItems
+{
+public:
+	// E.g. allow hash table to rehash after all items are loaded
+	void RequestFinalCall() { bNeedFinalize = true; }
+
+	void SetUnconstructed() { bUnconstructed = true; }
+
+	// Non-contiguous items must be set individually
+	template<typename ItemType>
+	void Set(ItemType* Items, uint64 NumItems)
+	{
+		Set(Items, Num, sizeof(ItemType));
+	}
+
+	void Set(void* Items, uint64 NumItems, uint32 ItemSize)
+	{
+		check(NumItems > 0);
+		check(Items != Data);
+		Data = reinterpret_cast<uint8*>(Items);
+		Num = NumItems;
+		Size = ItemSize;
+	}
+
+	template<typename ItemType>
+	TArrayView<ItemType> Get()
+	{ 
+		return MakeArrayView(reinterpret_cast<ItemType*>(Data), Num);
+	}
+
+private:
+	friend class FRangeLoader;
+	uint8*	Data = nullptr;
+	uint64	Num = 0;			
+	uint32	Size = 0;
+	bool	bNeedFinalize = false;
+	bool	bUnconstructed = false;
+
+	uint64	NumBytes() const { return Num * Size; }
+};
+
+struct FLoadRangeContext
+{
+	FConstructionRequest	Request;		// Request to construct items to be loaded
+	FConstructedItems		Items;			// Response from IRangeBinding
+	uint64					Scratch[64];	// Scratch memory for IRangeBinding
+};
+
+// todo: switch to class
+struct FGetItemsRequest
+{
+	template<typename T>
+	const T& GetRange() const { return *reinterpret_cast<const T*>(Range); }
+
+	const void* Range = nullptr;
+	uint64 Index = 0;
+};
+
+struct FExistingItemSlice
+{
+	const void*		Data = nullptr;			
+	uint64			Num = 0;
+
+	const uint8* At(uint64 Idx, uint32 Stride) const
+	{
+		check(Idx < Num);
+		return reinterpret_cast<const uint8*>(Data) + Idx*Stride;
+	}
+};
+
+struct FExistingItems
+{
+	uint64				NumTotal = 0;
+	uint32				Stride = 0;
+	FExistingItemSlice	Slice;
+
+	void SetPart(FExistingItemSlice Part)
+	{
+		NumTotal += Part.Num;
+		Slice = Part;
+	}
+
+	void SetAll(FExistingItemSlice Whole, uint32 InStride)
+	{
+		NumTotal = Whole.Num;
+		Stride = InStride;
+		Slice = Whole;
+	}
+
+	template<typename ItemType>
+	void SetAll(const ItemType* Items, uint64 NumItems)
+	{
+		SetAll(Items, NumItems, sizeof(ItemType));
+	}
+
+	//bool HasAll() const
+	//{
+	//	return NumTotal == Slice.Num;
+	//}
+};
+
+struct FSaveRangeContext
+{
+	FGetItemsRequest		Request;	// Request to get items to be saved
+	FExistingItems			Items;		// Response from IRangeBinding
+	uint64					Scratch[8]; // Scratch memory for IRangeBinding
+};
+
+class IRangeBinding
+{
+public:
+	virtual void MakeItems(FLoadRangeContext& Ctx) const = 0;
+	virtual void ReadItems(FSaveRangeContext& Ctx) const = 0;
+};
+
+////////////////////////////////////////////////////////////////////////////////////////////////
+
+class FRangeBinding
+{
+	static constexpr uint64		SizeMask = 0b111;
+	uint64						Handle;
+
+public:
+	FRangeBinding(const IRangeBinding& Binding, ERangeSizeType SizeType);
+
+	const IRangeBinding&		GetBinding() const	{ return *reinterpret_cast<IRangeBinding*>(Handle & ~SizeMask); }
+	ERangeSizeType				GetSizeType() const { return static_cast<ERangeSizeType>(Handle & SizeMask); }
+};
+
+template<typename T>
+struct TRangeBind{ using Type = void; };
+
+template<typename T>
+using RangeBind = typename TRangeBind<T>::Type;
+
+////////////////////////////////////////////////////////////////////////////////////////////////
+
+struct FMemberBinding
+{
+	uint64							Offset;
+	FMemberBindType					InnermostType;		// Always Leaf or Struct
+	FOptionalSchemaId				InnermostSchema;	// Enum or struct schema
+	TConstArrayView<FRangeBinding>	RangeBindings;		// Non-empty -> Range
+};
+
+////////////////////////////////////////////////////////////////////////////////////////////////
+
+
+class FStructBindings
+{
+public:
+	UE_NONCOPYABLE(FStructBindings);
+	FStructBindings() = default;
+	~FStructBindings();
+
+	void						BindStruct(FStructSchemaId Id, const ICustomStructBinding& Custom);
+	void						BindStruct(FStructSchemaId Id, TConstArrayView<FMemberBinding> Schema);
+	FStructBinding				Get(FStructSchemaId Id) const;
+	void						DropStruct(FStructSchemaId Id);
+
+private:
+	TArray<FStructBinding>		Bindings;
+
+	void						Bind(FStructSchemaId Id, FStructBinding Binding);
+};
+
+////////////////////////////////////////////////////////////////////////////////////////////////
+
+//struct FNamedMemberBinding : FMemberBinding
+//{
+//	FMemberId					Name;
+//};
+
+template<typename Type, class Ids>
+FMemberBindType BindMemberLeaf(FOptionalSchemaId& OutSchema)
+{
+	OutSchema = std::is_enum_v<Type> ? IndexNativeEnum<Type, Ids>() : NoId;
+	return ReflectLeaf<Type>;
+}
+
+
+template<typename CustomBinding, class Runtime>
+FStructSchemaId BindCustomStructOnce()
+{
+	struct FBinding
+	{
+		using Type = typename CustomBinding::Type;
+		using Ids = typename Runtime::Ids;
+
+		FBinding()
+		: Id(Runtime::GetTypes().DeclareStruct(Ids::IndexNativeType(Type::Name), NoId, CustomBinding::GetMemberIds(), CustomBinding::Occupancy))
+		{
+			Runtime::GetCustomBindings().BindStruct(Id, Instance);
+		}
+
+		~FBinding()
+		{
+			Runtime::DropStruct(Id);
+		}
+
+		CustomBinding Instance;
+		FStructSchemaId Id;
+	};
+
+	static FBinding Binding;
+	return Binding.Id;
+}
+
+template<class Type, class CustomBinding, class Runtime>
+FMemberBindType BindMemberStruct(FOptionalSchemaId& OutSchema)
+{
+	FStructSchemaId Id	= std::is_void_v<CustomBinding>
+						? IndexNativeStruct<Type, typename Runtime::Ids>()
+						: BindCustomStructOnce<CustomBinding, Runtime>();
+
+	OutSchema = FOptionalSchemaId(Id);
+	return FMemberBindType(FStructType{EMemberKind::Struct, /* IsDynamic */ 0, /* IsSuper */ 0});
+}
+
+template<typename Type, class Ids>
+FMemberBindType BindType(FOptionalSchemaId& OutSchema)
+{
+	if constexpr (std::is_arithmetic_v<Type> || std::is_enum_v<Type>)
+	{
+		OutSchema = std::is_enum_v<Type> ? IndexNativeEnum<Type, Ids>() : NoId;
+		return ReflectLeaf<Type>;
+	}
+	else
+	{
+		OutSchema = FOptionalSchemaId(IndexNativeStruct<Type, Ids>());
+		return FMemberBindType(FStructType{EMemberKind::Struct, /* IsDynamic */ 0, /* IsSuper */ 0});
+	}
+}
+
+template<typename RangeBinding>
+constexpr uint32 CountRangeBindings()
+{
+	if constexpr (std::is_void_v<RangeBinding>)
+	{
+		return 0;
+	}
+	else
+	{
+		return 1 + CountRangeBindings<typename RangeBinding::ItemType>();
+	}
+}
+
+template<typename RangeBinding>
+struct TGetInnermost
+{
+	using InnerType = typename RangeBinding::ItemType;
+    using Type = std::conditional_t<std::is_void_v<RangeBinding>, InnerType, typename TGetInnermost<InnerType>::Type>;
+};
+
+template<typename RangeBinding, uint32 N>
+TConstArrayView<FRangeBinding> GetRangeBindings()
+{
+	static_assert(N != 0);
+	if constexpr (N == 1)
+	{
+		static RangeBinding StaticInstance;
+		return {&StaticInstance, 1};
+	}
+	else
+	{
+		struct FNestedBindings
+		{
+			FNestedBindings()
+			{
+				Instances[0] = &Instance;
+				FMemory::Memcpy(Instances + 1, GetRangeBindings<typename RangeBinding::ItemType, N - 1>().GetData(), (N - 1) * sizeof(FRangeBinding*) );
+			}
+			RangeBinding Instance;
+			FRangeBinding* Instances[N];
+		};
+
+		static FNestedBindings Static;
+		return MakeArrayView(Static.Instances);
+	}
+}
+
+template<class Var, class Runtime>
+FMemberBinding BindMember()
+{
+	using Ids = typename Runtime::Ids;
+	using Type = Var::Type;
+
+	FMemberBindType Dummy(ERangeSizeType);
+	FMemberBinding Out = { Var::Offset, Dummy };
+	if constexpr (std::is_arithmetic_v<Type> || std::is_enum_v<Type>)
+	{
+		Out.InnermostType = BindMemberLeaf<Type, Ids>(Out.InnermostSchema);
+	}
+	else
+	{
+		using CustomBinding = typename Runtime::template CustomBindings<Type>::Type;
+		using RangeBinding = RangeBind<Type>;
+		constexpr uint32 NumRangeBindings = CountRangeBindings<RangeBinding>();
+
+		if constexpr (std::is_void_v<CustomBinding> && NumRangeBindings)
+		{
+			using InnermostType = typename TGetInnermost<RangeBinding>::Type;
+			Out.RangeBindings = GetRangeBindings<RangeBinding, NumRangeBindings>();
+			Out.InnermostType = BindType<InnermostType, Ids>(Out.InnermostSchema);
+		}
+		else
+		{
+			Out.InnermostType = BindMemberStruct<Type, CustomBinding, Runtime>(Out.InnermostSchema);
+		}
+	}
+	
+	return Out;
+}
+
+template<typename Enum, typename Ids>
+FEnumSchemaId IndexNativeEnum()
+{
+	static FEnumSchemaId Id = Ids::IndexEnum(CttiOf<Enum>::Name);
+	return Id;
+}
+
+template<typename Struct, typename Ids>
+FStructSchemaId IndexNativeStruct()
+{
+	static FStructSchemaId Id = Ids::IndexStruct(CttiOf<Struct>::Name);
+	return Id;
+}
+
+
+template<typename Struct, typename Ids>
+FOptionalStructSchemaId IndexOptionalNativeStruct()
+{
+	if constexpr (!std::is_void_v<Struct>)
+	{
+		return IndexNativeStruct<Struct, Ids>();
+	}
+	
+	return NoId;
+}
+
+
+template<class Ctti, class Ids>
+FEnumSchemaId DeclareNativeEnum(FDeclarations& Out, EEnumMode Mode)
+{
+	using UnderlyingType = std::underlying_type_t<typename Ctti::Type>;
+	FTypeId Type = Ids::IndexNativeType(Ctti::Name);	
+	FEnumerator Enumerators[Ctti::NumEnumerators];
+	for (FEnumerator& Enumerator : Enumerators)
+	{
+		Enumerator.Name = Ids::IndexMember(Ctti::Enumerators[&Enumerator - Enumerators].Name);
+		Enumerator.Constant = static_cast<uint64>(static_cast<UnderlyingType>(Ctti::Enumerators[&Enumerator - Enumerators].Constant));
+	}
+
+	return Out.DeclareEnum(Type, Mode, LeafWidth<sizeof(UnderlyingType)>, Enumerators);
+}
+
+//
+//template<class Ctti, class Rttis>
+//FStructSchemaId BindNativeStruct(FDeclarations& Runtime, EMemberPresence Occupancy)
+//{
+//	using Ids = typename Rttis::IdsType;
+//	FTypeId Type = Ids::IndexNativeType(Ctti::Name);
+//
+//	FMemberBinding MemberBindings[Ctti::NumVars];
+//	FMemberId MemberIds[Ctti::NumVars];
+//	ForEachVar<Ctti>([&]<class Var>()
+//	{ 
+//		MemberIds[Var::Index] = Ids::IndexMember(Var::Name);
+//		MemberBindings[Var::Index] = BindMember<Var, Rttis>();
+//	});
+//
+//	FOptionalSchemaId SuperId = !std::is_void_v<Ctti::Super> ? IndexNativeStruct<Ctti::Super, Ids>() : NoId;
+//	FStructSchemaId Id = Runtime.DeclareStruct(Type, SuperId, MemberIds, Occupancy);
+//	Runtime.BindStruct(Id, MemberBindings);
+//	return Id;
+//}
+
+
+template<class Ctti, class Ids>
+FStructSchemaId DeclareNativeStruct(FDeclarations& Out, EMemberPresence Occupancy)
+{
+	FTypeId Type = Ids::IndexNativeType(Ctti::Name);
+	FStructSchemaId Id = Ids::IndexStruct(Type);
+	FOptionalStructSchemaId SuperId = IndexOptionalNativeStruct<typename Ctti::Super, Ids>();
+	FMemberId MemberIds[Ctti::NumVars];
+	ForEachVar<Ctti>([&]<class Var>()
+	{ 
+		MemberIds[Var::Index] = Ids::IndexMember(Var::Name);
+	});
+	Out.DeclareStruct(Id, Type, MemberIds, Occupancy, SuperId);
+
+	return Id;
+}
+
+template<class Ctti, class Runtime>
+void BindNativeStruct(FStructBindings& Out, FStructSchemaId DeclaredId)
+{
+	FMemberBinding MemberBindings[Ctti::NumVars];
+	ForEachVar<Ctti>([&]<class Var>()
+	{ 
+		MemberBindings[Var::Index] = BindMember<Var, Runtime>();
+	});
+	Out.BindStruct(DeclaredId, MemberBindings);
+}
+
+////////////////////////////////////////////////////////////////////////////////////////////////
+
+} // namespace PlainProps
+
diff --git a/Engine/Source/Programs/PlainPropsTests/Public/PlainPropsBuild.h b/Engine/Source/Programs/PlainPropsTests/Public/PlainPropsBuild.h
new file mode 100644
--- /dev/null
+++ b/Engine/Source/Programs/PlainPropsTests/Public/PlainPropsBuild.h
@@ -0,0 +1,238 @@
+// Copyright Epic Games, Inc. All Rights Reserved.
+
+#pragma once
+
+#include  "PlainPropsDeclare.h"
+#include  "PlainPropsTypes.h"
+#include "Containers/Array.h"
+#include "Containers/ArrayView.h"
+#include "Memory/MemoryView.h"
+#include "Templates/UniquePtr.h"
+
+namespace PlainProps
+{
+
+struct FBuiltMember;
+struct FBuiltStruct;
+struct FBuiltRange;
+class FDebugIds;
+struct FUnpackedLeafType;
+
+//////////////////////////////////////////////////////////////////////////
+
+struct FMemberSchema
+{
+	FMemberType				Type;
+	FOptionalSchemaId		InnerSchema;
+	TArray<FMemberType>		InnerRangeTypes;
+
+	FMemberType GetInnermostType() const { return Type.IsRange() ? InnerRangeTypes.Last() : Type; }
+};
+
+inline bool operator==(const FMemberSchema& A, const FMemberSchema& B)
+{
+	return A.Type == B.Type && A.InnerSchema == B.InnerSchema && A.InnerRangeTypes == B.InnerRangeTypes;
+}
+//////////////////////////////////////////////////////////////////////////
+
+inline uint64 ValueCast(bool Value)			{ return static_cast<uint64>(Value); }
+inline uint64 ValueCast(int8 Value)			{ return static_cast<uint8>(Value); }
+inline uint64 ValueCast(int16 Value)		{ return static_cast<uint16>(Value); }
+inline uint64 ValueCast(int32 Value)		{ return static_cast<uint32>(Value); }
+inline uint64 ValueCast(int64 Value)		{ return static_cast<uint64>(Value); }
+inline uint64 ValueCast(uint8 Value)		{ return Value; }
+inline uint64 ValueCast(uint16 Value)		{ return Value; }
+inline uint64 ValueCast(uint32 Value)		{ return Value; }
+inline uint64 ValueCast(uint64 Value)		{ return Value; }
+uint64 ValueCast(float Value);
+uint64 ValueCast(double Value);
+inline uint64 ValueCast(char8_t Value)		{ return static_cast<uint8>(Value); }
+inline uint64 ValueCast(char16_t Value)		{ return static_cast<uint16>(Value); }
+inline uint64 ValueCast(char32_t Value)		{ return static_cast<uint32>(Value); }
+	
+inline constexpr ERangeSizeType RangeSizeOf(bool)	{ return ERangeSizeType::Uni; }
+inline constexpr ERangeSizeType RangeSizeOf(int8)	{ return ERangeSizeType::S8; }
+inline constexpr ERangeSizeType RangeSizeOf(int16)	{ return ERangeSizeType::S16; }
+inline constexpr ERangeSizeType RangeSizeOf(int32)	{ return ERangeSizeType::S32; }
+inline constexpr ERangeSizeType RangeSizeOf(int64)	{ return ERangeSizeType::S64; }
+inline constexpr ERangeSizeType RangeSizeOf(uint8)	{ return ERangeSizeType::U8; }
+inline constexpr ERangeSizeType RangeSizeOf(uint16)	{ return ERangeSizeType::U16; }
+inline constexpr ERangeSizeType RangeSizeOf(uint32)	{ return ERangeSizeType::U32; }
+inline constexpr ERangeSizeType RangeSizeOf(uint64)	{ return ERangeSizeType::U64; }
+
+//////////////////////////////////////////////////////////////////////////
+
+// Todo: Turn into smart pointer where dtor call FBuiltRange::Delete if Values not detached to FBuiltValue
+struct FTypedRange
+{
+	FMemberSchema Schema;
+	FBuiltRange* Values = nullptr;
+};
+
+
+template<typename LeafType, typename SizeType>
+FMemberSchema MakeLeafRangeSchema()
+{
+	check(std::is_arithmetic_v<LeafType>);
+	return { FMemberType(RangeSizeOf(SizeType{})), NoId, { ReflectLeaf<LeafType>.Pack() } };
+}
+
+template<typename EnumType, typename SizeType>
+FMemberSchema MakeEnumRangeSchema(FEnumSchemaId Schema)
+{
+	check(std::is_enum_v<EnumType>);
+	return { FMemberType(RangeSizeOf(SizeType{})), Schema, { ReflectLeaf<EnumType>.Pack() } };
+}
+
+inline constexpr FMemberType DefaultStructType = FMemberType(FStructType{EMemberKind::Struct, /* IsDynamic */ 0, /* IsSuper */ 0});
+inline constexpr FMemberType SuperStructType =	 FMemberType(FStructType{EMemberKind::Struct, /* IsDynamic */ 0, /* IsSuper */ 1});
+
+
+inline FMemberSchema MakeStructRangeSchema(ERangeSizeType SizeType, FStructSchemaId Schema)
+{
+	return { FMemberType(SizeType), Schema, { DefaultStructType } };
+}
+
+PP_API FMemberSchema MakeNestedRangeSchema(ERangeSizeType SizeType, const FMemberSchema& InnerRangeSchema);
+
+namespace Private
+{
+	PP_API [[nodiscard]] FBuiltRange* BuildStructuralRange(/* in-out */ TArrayView64<FBuiltRange*> Structs);
+	PP_API [[nodiscard]] FBuiltRange* BuildStructuralRange(/* in-out */ TArrayView64<TUniquePtr<FBuiltStruct>> Structs);
+	PP_API [[nodiscard]] FBuiltRange* BuildLeafRange(FUnpackedLeafType Leaf, uint64 Num, FMemoryView Values);
+	PP_API void						NormalizeLeafRange(FUnpackedLeafType Leaf, FBuiltRange& Out);
+
+	template<typename LeafType, typename SizeType>
+	[[nodiscard]] FBuiltRange* BuildLeafRange(const LeafType* Values, SizeType InNum)
+	{
+		static_assert(SizeOf(ReflectLeaf<LeafType>.Width) == sizeof(LeafType));
+
+		uint64 Num = IntCastChecked<uint64>(InNum);
+		FBuiltRange* Range = Num > 0 ? Private::BuildLeafRange(ReflectLeaf<LeafType>, Num, MakeMemoryView(Values, Num * sizeof(LeafType))) : nullptr;
+	
+		if (Range && std::is_floating_point_v<LeafType>)
+		{
+			NormalizeLeafRange(ReflectLeaf<LeafType>, *Range);
+		}
+	
+		return Range;
+	}
+}
+
+template<typename LeafType, typename SizeType>
+[[nodiscard]] FTypedRange BuildLeafRange(const LeafType* Values, SizeType Num)
+{
+	return { MakeLeafRangeSchema<LeafType, SizeType>(), Private::BuildLeafRange(Values, Num) };
+}
+
+template<typename LeafType, typename SizeType>
+[[nodiscard]] FTypedRange BuildLeafRange(TConstArrayView<LeafType, SizeType> Values)
+{
+	return BuildLeafRange(Values.GetData(), Values.Num());
+}
+
+template<typename EnumType, typename SizeType>
+[[nodiscard]] FTypedRange BuildEnumRange(FEnumSchemaId Enum, TConstArrayView<EnumType, SizeType> Values)
+{
+	return { MakeEnumRangeSchema<EnumType, SizeType>(Enum), Private::BuildLeafRange(Values.GetData(), Values.Num()) };
+}
+
+[[nodiscard]] inline FTypedRange BuildStructRange(FStructSchemaId Schema, ERangeSizeType SizeType, /* in-out */ TArrayView64<TUniquePtr<FBuiltStruct>> Values )
+{
+	return { MakeStructRangeSchema(SizeType, Schema), Values.Num() ? Private::BuildStructuralRange(/* ownership xfer */ Values) : nullptr };
+}
+
+// Builds an ordered list of properties to be saved
+class FMemberBuilder
+{
+public:
+	template<typename LeafType>
+	void Add(FMemberId Name, LeafType Value)
+	{
+		static_assert(std::is_arithmetic_v<LeafType>, "Illegal leaf type");
+		AddLeaf(Name, ReflectLeaf<LeafType>, {}, ValueCast(Value));
+	}
+
+	template<typename EnumType>
+	void AddEnum(FMemberId Name, FEnumSchemaId Schema, EnumType Value)
+	{
+		AddLeaf(Name, ReflectLeaf<EnumType>, ToOptional(Schema), ValueCast(Value));
+	}
+	
+	void AddEnum8(FMemberId Name, FEnumSchemaId Schema, uint8 Value)	{ AddLeaf(Name, {ELeafType::Enum, ELeafWidth::B8},  ToOptional(Schema), Value); }
+	void AddEnum16(FMemberId Name, FEnumSchemaId Schema, uint16 Value)	{ AddLeaf(Name, {ELeafType::Enum, ELeafWidth::B16}, ToOptional(Schema), Value); }
+	void AddEnum32(FMemberId Name, FEnumSchemaId Schema, uint32 Value)	{ AddLeaf(Name, {ELeafType::Enum, ELeafWidth::B32}, ToOptional(Schema), Value); }
+	void AddEnum64(FMemberId Name, FEnumSchemaId Schema, uint64 Value)	{ AddLeaf(Name, {ELeafType::Enum, ELeafWidth::B64}, ToOptional(Schema), Value); }
+
+	PP_API void AddLeaf(FMemberId Name, FUnpackedLeafType Leaf, FOptionalEnumSchemaId Enum, uint64 Value);
+	PP_API void AddStruct(FMemberId Name, FStructSchemaId Schema, TUniquePtr<FBuiltStruct>&& Struct);	
+	PP_API void AddRange(FMemberId Name, FTypedRange&& Range);
+	
+	// Build members into a single nested super struct member, no-op if no non-super members has been added
+	PP_API void BuildSuperStruct(const FStructDeclaration& Super, const FDebugIds& Debug);
+
+	PP_API [[nodiscard]] TUniquePtr<FBuiltStruct> BuildAndReset(const FStructDeclaration& Declared, const FDebugIds& Debug);
+
+private:
+	TArray<FBuiltMember>	Members;
+	
+//	PP_API uint8* AllocateLeafRange(FMemberId Name, FUnpackedLeafType Leaf, ERangeSizeType RangeMax, uint64 Num);
+	
+	//template<typename T>
+	//void NormalizeLeafRange(T*, uint64) {}
+	//PP_API void NormalizeLeafRange(float*, uint64 Num);
+	//PP_API void NormalizeLeafRange(double*, uint64 Num);
+	
+};
+
+// Helper class for building struct ranges
+class FStructRangeBuilder
+{
+public:
+	FStructRangeBuilder(uint64 Num, ERangeSizeType InSizeType)
+	: SizeType(InSizeType)
+	{
+		Structs.SetNum(Num);
+	}
+
+	template<typename IntType>
+	explicit FStructRangeBuilder(IntType Num)
+	: FStructRangeBuilder(static_cast<uint64>(Num), RangeSizeOf(Num))
+	{}
+	 
+	FMemberBuilder& operator[](uint64 Idx) { return Structs[Idx]; }
+
+	FTypedRange BuildAndReset(const FStructDeclaration& Declared, const FDebugIds& Debug);
+
+private:
+	TArray64<FMemberBuilder> Structs; 
+	ERangeSizeType SizeType;
+};
+
+
+// Helper class for building nested ranges
+class FNestedRangeBuilder
+{
+public:
+	FNestedRangeBuilder(FMemberSchema InSchema, int64 InitialReserve)
+	: Schema(InSchema)
+	{
+		Ranges.Reserve(InitialReserve);
+	}
+
+	~FNestedRangeBuilder();
+
+	void Add(FTypedRange Range)
+	{
+		check(Range.Schema == Schema || Range.Values == nullptr);
+		Ranges.Add(Range.Values);
+	}
+
+	[[nodiscard]] FTypedRange BuildAndReset(ERangeSizeType SizeType);
+
+private:
+	TArray64<FBuiltRange*> Ranges; 
+	FMemberSchema Schema;
+};
+
+} // namespace PlainProps
\ No newline at end of file
diff --git a/Engine/Source/Programs/PlainPropsTests/Public/PlainPropsBuildSchema.h b/Engine/Source/Programs/PlainPropsTests/Public/PlainPropsBuildSchema.h
new file mode 100644
--- /dev/null
+++ b/Engine/Source/Programs/PlainPropsTests/Public/PlainPropsBuildSchema.h
@@ -0,0 +1,90 @@
+// Copyright Epic Games, Inc. All Rights Reserved.
+
+#pragma once
+
+#include  "PlainPropsBuild.h"
+#include "Containers/Array.h"
+#include "Containers/ArrayView.h"
+#include "Containers/Map.h"
+
+namespace PlainProps
+{
+
+struct FMemberSchema;
+
+struct FBuiltStructSchema
+{
+	FTypeId Type;
+	bool bUsed = false;
+	bool bDense = false;
+	FOptionalStructSchemaId Super;
+	TArray<FMemberId> MemberNames;
+	TArray<const FMemberSchema*> MemberSchemas;
+};
+
+struct FBuiltEnumSchema
+{
+	FTypeId Type;
+	bool bUsed = false;
+	EEnumMode Mode = EEnumMode::Flat;
+	ELeafWidth Width = ELeafWidth::B8;
+	TArray<FNameId> Names;
+	TArray<uint64> Constants;
+};
+
+struct FBuiltSchemas
+{
+	TArray<FBuiltStructSchema> Structs; // Same size as number of declared structs
+	TArray<FBuiltEnumSchema> Enums; // Same size as number of declared enums
+};
+
+class FSchemasBuilder;
+
+struct FStructSchemaBuilder
+{
+	const FStructDeclaration&		Declaration;
+	FSchemasBuilder&				AllSchemas;
+	const FDebugIds&				Debug;
+	TMap<FOptionalMemberId, FMemberSchema>	NotedMembers;
+	bool							bMissingMemberNoted = false;
+	mutable bool					bBuilt = false;
+
+	void							NoteMembersRecursively(const FBuiltStruct& Struct);
+	void							NoteRangeRecursively(ERangeSizeType NumType, TConstArrayView<FMemberType> Types, FSchemaId InnermostSchema, const FBuiltRange& Range);
+	FBuiltStructSchema				Build() const;
+};
+
+struct FEnumSchemaBuilder
+{
+	const FEnumDeclaration&			Declaration;
+	TSet<uint64>					NotedConstants;
+	mutable bool					bBuilt = false;
+
+	void							NoteValue(uint64 Value);
+	FBuiltEnumSchema				Build() const;
+};
+
+class FSchemasBuilder
+{
+public:
+	FSchemasBuilder(TConstArrayView<TUniquePtr<FStructDeclaration>> InStructs, TConstArrayView<TUniquePtr<FEnumDeclaration>> InEnums, const FDebugIds& InDebug);
+
+	void NoteMembers(FStructSchemaId Schema, const FBuiltStruct& Struct)
+	{
+		Structs[Schema.Idx].NoteMembersRecursively(Struct);
+	}
+
+	void NoteValue(FEnumSchemaId Schema, uint64 Value)
+	{
+		Enums[Schema.Idx].NoteValue(Value);
+	}
+	
+	FBuiltSchemas Build() const;
+
+private:
+	TArray<FStructSchemaBuilder> Structs;
+	TArray<FEnumSchemaBuilder> Enums;
+	const FDebugIds& Debug;
+};
+
+} // namespace PlainProps
\ No newline at end of file
diff --git a/Engine/Source/Programs/PlainPropsTests/Public/PlainPropsCtti.h b/Engine/Source/Programs/PlainPropsTests/Public/PlainPropsCtti.h
new file mode 100644
--- /dev/null
+++ b/Engine/Source/Programs/PlainPropsTests/Public/PlainPropsCtti.h
@@ -0,0 +1,94 @@
+// Copyright Epic Games, Inc. All Rights Reserved.
+
+#pragma once
+
+#include <cstdint>
+
+namespace PlainProps 
+{
+
+#define _PP_ARITHMETIC_CTTI(T) struct T##_Ctti  { inline static constexpr char Name[] = #T; using Type = T; }; T##_Ctti CttiOfPtr(T*);
+_PP_ARITHMETIC_CTTI(float)
+_PP_ARITHMETIC_CTTI(double)
+_PP_ARITHMETIC_CTTI(bool)
+_PP_ARITHMETIC_CTTI(int8_t)
+_PP_ARITHMETIC_CTTI(uint8_t)
+_PP_ARITHMETIC_CTTI(int16_t)
+_PP_ARITHMETIC_CTTI(int32_t)
+_PP_ARITHMETIC_CTTI(uint32_t)
+_PP_ARITHMETIC_CTTI(int64_t)
+_PP_ARITHMETIC_CTTI(uint64_t)
+_PP_ARITHMETIC_CTTI(char)
+_PP_ARITHMETIC_CTTI(char8_t)
+_PP_ARITHMETIC_CTTI(char16_t)
+_PP_ARITHMETIC_CTTI(char32_t)
+_PP_ARITHMETIC_CTTI(wchar_t)
+#undef _PP_ARITHMETIC_CTTI
+
+//////////////////////////////////////////////////////////////////////////
+
+template<typename T> struct TCttiOf { using Type = decltype(CttiOfPtr((T*)nullptr)); };
+template<typename T> using CttiOf = typename TCttiOf<T>::Type;
+
+//////////////////////////////////////////////////////////////////////////
+// Private PP_REFLECT_XYZ helpers
+
+#define _PP_EXPAND(X) X
+#define _PP_NUM_ARGS_(_1,_2,_3,_4,_5,_6,_7,_8,_9,_10, N, ...) N
+#define _PP_NUM_ARGS(...) _PP_EXPAND(_PP_NUM_ARGS_(__VA_ARGS__,10,9,8,7,6,5,4,3,2,1))
+
+#define _PP_REFLECT_1(N, MACRO, NS, T, M)				MACRO(N-1, NS, T, M)
+#define _PP_REFLECT_2(N, MACRO, NS, T, M, ...)			MACRO(N-2, NS, T, M) _PP_REFLECT_1(N, MACRO, NS, T, __VA_ARGS__)
+#define _PP_REFLECT_3(N, MACRO, NS, T, M, ...)			MACRO(N-3, NS, T, M) _PP_REFLECT_2(N, MACRO, NS, T, __VA_ARGS__)
+#define _PP_REFLECT_4(N, MACRO, NS, T, M, ...)			MACRO(N-4, NS, T, M) _PP_REFLECT_3(N, MACRO, NS, T, __VA_ARGS__)
+#define _PP_REFLECT_5(N, MACRO, NS, T, M, ...)			MACRO(N-5, NS, T, M) _PP_REFLECT_4(N, MACRO, NS, T, __VA_ARGS__)
+#define _PP_REFLECT_6(N, MACRO, NS, T, M, ...)			MACRO(N-6, NS, T, M) _PP_REFLECT_5(N, MACRO, NS, T, __VA_ARGS__)
+#define _PP_REFLECT_7(N, MACRO, NS, T, M, ...)			MACRO(N-7, NS, T, M) _PP_REFLECT_6(N, MACRO, NS, T, __VA_ARGS__)
+#define _PP_REFLECT_8(N, MACRO, NS, T, M, ...)			MACRO(N-8, NS, T, M) _PP_REFLECT_7(N, MACRO, NS, T, __VA_ARGS__)
+#define _PP_REFLECT_9(N, MACRO, NS, T, M, ...)			MACRO(N-9, NS, T, M) _PP_REFLECT_8(N, MACRO, NS, T, __VA_ARGS__)
+#define _PP_REFLECT_THINGS(N, MACRO, ...)				_PP_EXPAND(_PP_REFLECT_##N(N, MACRO, __VA_ARGS__))
+
+#define _PP_REFLECT_ENUMERATOR(N, NS, T, C)				{ #C, Type::C }, 
+#define _PP_REFLECT_ENUM_INNER(N, NS, T, ES)			struct T##_Ctti { inline static constexpr char Name[] = #T; using Type = NS :: T; static constexpr int NumEnumerators = N; static constexpr struct { const char* Name; Type Value; } Enumerators[] = { ES }; }; T##_Ctti CttiOfPtr(T*);
+#define _PP_REFLECT_ENUM(N, NS, T, ...)					_PP_REFLECT_ENUM_INNER(N, NS, T, _PP_EXPAND(_PP_REFLECT_THINGS(N, _PP_REFLECT_ENUMERATOR, NS, T, __VA_ARGS__)))					
+#define _PP_REFLECT_STRUCT(N, NS, T, S, ...)			PP_REFLECT_STRUCT_ONLY(N, NS, T, S)				_PP_EXPAND(_PP_REFLECT_THINGS(N, PP_REFLECT_MEMBER, NS, T, __VA_ARGS__))
+#define _PP_REFLECT_STRUCT_TEMPLATE(N, NS, T, S, ...)	PP_REFLECT_STRUCT_TEMPLATE_ONLY(N, NS, T, S)	_PP_EXPAND(_PP_REFLECT_THINGS(N, PP_REFLECT_TEMPLATE_MEMBER, NS, T, __VA_ARGS__))
+
+//////////////////////////////////////////////////////////////////////////
+
+#define PP_REFLECT_ENUM(NS, T, ...)						_PP_REFLECT_ENUM(_PP_NUM_ARGS(__VA_ARGS__), NS, T, __VA_ARGS__)
+#define PP_REFLECT_STRUCT(NS, T, S, ...)				_PP_REFLECT_STRUCT(_PP_NUM_ARGS(__VA_ARGS__), NS, T, S, __VA_ARGS__)
+#define PP_REFLECT_STRUCT_TEMPLATE(NS, T, S, ...)		_PP_REFLECT_STRUCT_TEMPLATE(_PP_NUM_ARGS(__VA_ARGS__), NS, T, S, __VA_ARGS__)
+
+// Alternate set of macros to reflect classes/structs with bitfield bool members, e.g. uint8 bOol : 1;
+
+#define PP_REFLECT_STRUCT_ONLY(N, NS, T, S)				struct T##_Ctti { inline static constexpr char Name[] = #T; using Type = NS :: T; using Super = S; static constexpr int NumVars = N; template<int> struct Var; }; T##_Ctti CttiOfPtr(T*);
+#define PP_REFLECT_STRUCT_TEMPLATE_ONLY(N, NS, T, S)	template<class... Ts> struct T##_Ctti { inline static constexpr char Name[] = #T; using Type = NS :: T<Ts...>; using Super = S; using TemplateArgs = std::tuple<Ts...>; static constexpr int NumVars = N; template<int, int> struct Var_; template<int I> using Var = Var_<I, 0>; }; template<class... Ts> T##_Ctti<Ts...> CttiOfPtr(T<Ts...>*);
+#define PP_REFLECT_MEMBER(N, NS, T, M)					template<> struct T##_Ctti::Var<N> { inline static constexpr char Name[] = #M; using Type = decltype(NS :: T :: M); static constexpr std::size_t Offset = offsetof(NS :: T, M); static constexpr auto Pointer = &NS :: T :: M; static constexpr int Index = N; };
+#define PP_REFLECT_BIT_MEMBER(N, NS, T, M, BYTEOFFSET, BITOFFSET) template<> struct T##_Ctti::Var<N> { inline static constexpr char Name[] = #M; using Type = decltype(NS :: T :: M); static constexpr std::size_t Offset = BYTEOFFSET; static constexpr int BitIndex = BITOFFSET; static constexpr int Index = N; };
+#define PP_REFLECT_TEMPLATE_MEMBER(N, NS, T, M)			template<class... Ts> template <int _> struct T##_Ctti<Ts...>::Var_<N, _> { inline static constexpr char Name[] = #M; using Type = decltype(NS :: T<Ts...> :: M); static constexpr std::size_t Offset = offsetof(NS :: T<Ts...>, M); static constexpr auto Pointer = &NS :: T<Ts...> :: M; static constexpr int Index = N; };
+#define PP_REFLECT_TEMPLATE_BIT_MEMBER					todo
+
+//////////////////////////////////////////////////////////////////////////
+
+// E.g. ForEachVar<CttiOf<MyClass>>([]<class Var>() { printf(Var::Name); });
+template <class Ctti, typename Fn>
+static constexpr void ForEachVar(Fn&& Visitor)
+{
+    constexpr int N = Ctti::NumVars;
+    if constexpr (N > 0) Visitor.template operator()<typename Ctti::template Var<0>>();
+    if constexpr (N > 1) Visitor.template operator()<typename Ctti::template Var<1>>();
+    if constexpr (N > 2) Visitor.template operator()<typename Ctti::template Var<2>>();
+    if constexpr (N > 3) Visitor.template operator()<typename Ctti::template Var<3>>();
+    if constexpr (N > 4) Visitor.template operator()<typename Ctti::template Var<4>>();
+    if constexpr (N > 5) Visitor.template operator()<typename Ctti::template Var<5>>();
+    if constexpr (N > 6) Visitor.template operator()<typename Ctti::template Var<6>>();
+    if constexpr (N > 7) Visitor.template operator()<typename Ctti::template Var<7>>();
+    if constexpr (N > 8) Visitor.template operator()<typename Ctti::template Var<8>>();
+    if constexpr (N > 9) Visitor.template operator()<typename Ctti::template Var<9>>();
+}
+
+//////////////////////////////////////////////////////////////////////////
+
+} // namespace PlainProps
+
diff --git a/Engine/Source/Programs/PlainPropsTests/Public/PlainPropsDeclare.h b/Engine/Source/Programs/PlainPropsTests/Public/PlainPropsDeclare.h
new file mode 100644
--- /dev/null
+++ b/Engine/Source/Programs/PlainPropsTests/Public/PlainPropsDeclare.h
@@ -0,0 +1,69 @@
+// Copyright Epic Games, Inc. All Rights Reserved.
+
+#pragma once
+
+#include  "PlainPropsTypes.h"
+#include "Containers/Array.h"
+#include "Containers/ArrayView.h"
+#include "Templates/UniquePtr.h"
+
+namespace PlainProps
+{
+struct FEnumerator
+{
+	FNameId					Name;
+	uint64					Constant;
+};
+
+enum class EEnumMode { Flat, Flag };
+
+struct FEnumDeclaration
+{
+	FTypeId					Type;			// Could be removed
+	EEnumMode				Mode;
+	ELeafWidth				Width;
+	uint16					NumEnumerators;
+	FEnumerator				Enumerators[0];	// Constants must be unique, no aliases allowed
+
+	TConstArrayView<FEnumerator> GetEnumerators() const
+	{
+		return MakeArrayView(Enumerators, NumEnumerators);
+	}
+};
+
+enum class EMemberPresence { RequireAll, AllowSparse };
+
+struct FStructDeclaration
+{
+	FStructSchemaId			Id;				// Could be removed, might allow declaration dedup among templated types
+	FTypeId					Type;			// Could be removed, might allow declaration dedup among templated types
+	FOptionalStructSchemaId	Super;
+	EMemberPresence			Occupancy;
+	uint16					NumMembers;
+	FMemberId				MemberOrder[0];
+
+	TConstArrayView<FMemberId> GetMemberOrder() const
+	{
+		return MakeArrayView(MemberOrder, NumMembers);
+	}
+};
+
+class FDeclarations
+{
+public:
+	void											DeclareEnum(FEnumSchemaId Id, FTypeId Type, EEnumMode Mode, ELeafWidth Width, TConstArrayView<FEnumerator> Enumerators);
+	void											DeclareStruct(FStructSchemaId Id, FTypeId Type, TConstArrayView<FMemberId> MemberOrder, EMemberPresence Occupancy, FOptionalStructSchemaId Super = {});
+	
+	const FEnumDeclaration&							Get(FEnumSchemaId Id) const		{ return *DeclaredEnums[Id.Idx]; }
+	const FStructDeclaration&						Get(FStructSchemaId Id) const	{ return *DeclaredStructs[Id.Idx]; }
+	
+	TConstArrayView<TUniquePtr<FEnumDeclaration>>	GetEnums() const	{ return DeclaredEnums; }
+	TConstArrayView<TUniquePtr<FStructDeclaration>>	GetStructs() const	{ return DeclaredStructs; }
+
+protected:
+	TArray<TUniquePtr<FEnumDeclaration>>			DeclaredEnums;
+	TArray<TUniquePtr<FStructDeclaration>>			DeclaredStructs;
+
+};
+
+} // namespace PlainProps
\ No newline at end of file
diff --git a/Engine/Source/Programs/PlainPropsTests/Public/PlainPropsIndex.h b/Engine/Source/Programs/PlainPropsTests/Public/PlainPropsIndex.h
new file mode 100644
--- /dev/null
+++ b/Engine/Source/Programs/PlainPropsTests/Public/PlainPropsIndex.h
@@ -0,0 +1,103 @@
+// Copyright Epic Games, Inc. All Rights Reserved.
+
+#pragma once
+
+#include  "PlainPropsTypes.h"
+#include "Containers/Set.h"
+
+namespace PlainProps
+{
+
+class FNestedScopeIndexer
+{
+public:
+	FNestedScopeId						Index(FNestedScope Scope);
+	FNestedScopeId						Index(FScopeId Outer, FFlatScopeId Inner)			{ return Index({Outer, Inner}); }
+	FNestedScope						Resolve(FNestedScopeId Id) const;
+	int32								Num() const											{ return Scopes.Num(); }
+
+	auto								begin() const										{ return Scopes.begin(); }
+	auto								end() const											{ return Scopes.end(); }
+private:
+	TSet<FNestedScope> Scopes;
+};
+
+//////////////////////////////////////////////////////////////////////////
+
+class FParametricTypeIndexer
+{
+public:
+	~FParametricTypeIndexer();
+
+	FParametricTypeId					Index(FParametricTypeView View);
+	/// @return View invalidated by calling Index() (switch Parameters to TPagedArray to avoid)
+	FParametricTypeView					Resolve(FParametricTypeId Id) const;
+	FParametricType						At(int32 Idx) const									{ return Types[Idx]; }
+	int32								Num() const											{ return Types.Num(); }
+	TConstArrayView<FParametricType>	GetAllTypes() const									{ return Types; }
+	TConstArrayView<FTypeId>			GetAllParameters() const							{ return Parameters; }
+
+private:
+	uint32								NumSlots = 0;
+	uint32*								Slots = nullptr;
+	TArray<FParametricType>				Types;
+	TArray<FTypeId>						Parameters;
+};
+
+//////////////////////////////////////////////////////////////////////////
+
+class FIdIndexerBase : public FDebugIds
+{
+public:
+	FScopeId							NestScope(FScopeId Outer, FFlatScopeId Inner);
+	FTypeId								MakeParametricType(FTypeId Type, TConstArrayView<FTypeId> Params);
+
+	FEnumSchemaId						IndexEnum(FTypeId Type);
+	FStructSchemaId						IndexStruct(FTypeId Type);
+	
+	virtual FNestedScope				Resolve(FNestedScopeId Id) const override final		{ return NestedScopes.Resolve(Id); }
+	virtual FParametricTypeView			Resolve(FParametricTypeId Id) const override final	{ return ParametricTypes.Resolve(Id); }
+	
+	const FNestedScopeIndexer&			GetNestedScopes() const								{ return NestedScopes; }
+	const FParametricTypeIndexer&		GetParametricTypes() const							{ return ParametricTypes; }
+
+	virtual uint32						NumNames() const = 0;
+
+protected:
+	FNestedScopeIndexer					NestedScopes;
+	FParametricTypeIndexer				ParametricTypes;
+	TSet<FTypeId>						Enums;
+	TSet<FTypeId>						Structs;
+};
+
+template<class NameType>
+void AppendString(FString& Out, const NameType& Str)
+{
+	Out.Append(Str);
+}
+
+template<class NameType>
+class TIdIndexer : public FIdIndexerBase
+{
+public:
+	template<typename T> FNameId		MakeName(T&& Str)									{ return { IntCastChecked<uint32>(Names.Add(NameType(Str)).AsInteger()) }; }
+	template<typename T> FMemberId		NameMember(T&& Name)								{ return { MakeName(Name) }; }
+	template<typename T> FScopeId		MakeScope(T&& Str)									{ return FScopeId(static_cast<FFlatScopeId>(MakeName(Str))); }
+	template<typename T> FScopeId		NestScope(FScopeId Outer, T&& Inner)				{ return FIdIndexerBase::NestScope(Outer, static_cast<FFlatScopeId>(MakeName(Inner))); }
+	template<typename T> FTypenameId	MakeTypename(T&& Name)								{ return FTypenameId(FConcreteTypenameId(MakeName(Name))); }
+	template<typename T> FTypeId		MakeType(T&& Scope, T&& Name)						{ return { MakeScope(Scope), MakeTypename(Name) }; }
+
+	NameType							ResolveName(FNameId Id)	const						{ return Names.Get(FSetElementId::FromInteger(Id.Idx)); }
+	virtual uint32						NumNames() const override final						{ return static_cast<uint32>(Names.Num()); }
+
+protected:
+	TSet<NameType>						Names;
+
+	virtual void						AppendDebugString(FString& Out, FNameId Id) const override final
+	{
+		AppendString(Out, ResolveName(Id));
+	}
+};
+
+
+} // namespace PlainProps
diff --git a/Engine/Source/Programs/PlainPropsTests/Public/PlainPropsLoad.h b/Engine/Source/Programs/PlainPropsTests/Public/PlainPropsLoad.h
new file mode 100644
--- /dev/null
+++ b/Engine/Source/Programs/PlainPropsTests/Public/PlainPropsLoad.h
@@ -0,0 +1,40 @@
+// Copyright Epic Games, Inc. All Rights Reserved.
+
+#pragma once
+
+#include "Containers/ArrayView.h"
+#include  "PlainPropsRead.h"
+#include "Memory/MemoryFwd.h"
+
+namespace PlainProps 
+{
+
+class FDeclarations;
+struct FLoadBatch;
+class FRangeBinding;
+class FStructBindings;
+struct FTypedRange;
+
+
+PP_API FLoadBatch* CreateLoadPlans(FReadBatchId ReadId, const FDeclarations& Declarations, const FStructBindings& Bindings);
+PP_API void DestroyLoadPlans(FLoadBatch* Batch);
+PP_API void LoadStruct(uint8* Dst, FByteReader Src, FStructSchemaId Id, const FLoadBatch& Batch);
+PP_API void LoadStruct(uint8* Dst, FStructView Src, const FLoadBatch& Batch);
+PP_API void ConstructAndLoadStruct(uint8* Dst, FByteReader Src, FStructSchemaId Id, const FLoadBatch& Batch);
+PP_API void ConstructAndLoadStruct(uint8* Dst, FStructView Src, const FLoadBatch& Batch);
+PP_API void LoadRange(uint8* Dst, FRangeView Src, ERangeSizeType MaxType, TConstArrayView<FRangeBinding> Bindings, const FLoadBatch& Batch);
+
+//using FItemConstructor = void (*)(FLoadRangeContext& Ctx);
+
+//struct FLoadRangeBinding
+//{
+//	FItemConstructor	MakeItems;
+//	const void*			UserContext;
+//};
+
+// Possible optimization - opt-in leaf range load method to avoid copying and repeated calls when target representation is not a range
+//template<typename T> using TDirectLeafLoader = void (*)(TRangeView<T> Items);
+//using FDirectBoolLoader = void (*)(FBoolRangeView Bits);
+
+} // namespace PlainProps
+
diff --git a/Engine/Source/Programs/PlainPropsTests/Public/PlainPropsRead.h b/Engine/Source/Programs/PlainPropsTests/Public/PlainPropsRead.h
new file mode 100644
--- /dev/null
+++ b/Engine/Source/Programs/PlainPropsTests/Public/PlainPropsRead.h
@@ -0,0 +1,621 @@
+// Copyright Epic Games, Inc. All Rights Reserved.
+
+#pragma once
+
+#include "Containers/ArrayView.h"
+#include "HAL/PlatformMemory.h"
+#include "Memory/MemoryView.h"
+#include "Misc/AssertionMacros.h"
+#include "Serialization/VarInt.h"
+#include  "PlainPropsTypes.h"
+
+namespace PlainProps
+{
+
+struct FEnumSchema;
+struct FIdBinding;
+struct FSchemaBatch;
+struct FStructSchema;
+
+//////////////////////////////////////////////////////////////////////////
+
+// Represents a batch currently being read from
+class FReadBatchId
+{
+	friend class FReadSchemaRegistry;
+	uint16 Idx = 0;
+};
+
+// @param Schemas must outlive read batch
+const FSchemaBatch*		ValidateSchemas(FMemoryView Schemas);
+FSchemaBatch*			CreateTranslatedSchemas(const FSchemaBatch* Schemas, FIdBinding NewIds);
+FReadBatchId			MountReadSchemas(const FSchemaBatch* Schemas);
+const FSchemaBatch*		UnmountReadSchemas(FReadBatchId Batch);
+void					DestroyTranslatedSchemas(const FSchemaBatch* Schemas);
+
+uint32					NumStructSchemas(FReadBatchId Batch);
+const FStructSchema&	ResolveStructSchema(FReadBatchId Batch, FStructSchemaId Id);
+const FEnumSchema&		ResolveEnumSchema(FReadBatchId Batch, FEnumSchemaId Id);
+FNestedScope			ResolveUntranslatedNestedScope(FReadBatchId Batch, FNestedScopeId Id);
+FParametricTypeView		ResolveUntranslatedParametricType(FReadBatchId Batch, FParametricTypeId Id);
+
+//////////////////////////////////////////////////////////////////////////
+
+class FByteReader
+{
+	const uint8* It = nullptr;
+#if DO_CHECK
+	const uint8* End = nullptr;
+#endif
+public:
+	FByteReader() = default;
+	explicit FByteReader(FMemoryView View) : FByteReader(static_cast<const uint8*>(View.GetData()), View.GetSize()) {}
+	FByteReader(const uint8* Data, uint64 NumBytes) : FByteReader(Data, Data + NumBytes) {}
+	FByteReader(const uint8* InBegin, const uint8* InEnd)
+	: It(InBegin)
+#if DO_CHECK
+	, End(InEnd)
+#endif
+	{}
+
+	[[nodiscard]] const uint8* GrabBytes(uint64 NumBytes)
+	{
+		check(It + NumBytes <= End);
+		const uint8* Out = It;
+		It += NumBytes;
+		return Out;
+	}
+
+	[[nodiscard]] FMemoryView GrabSlice(uint64 NumBytes)
+	{
+		return FMemoryView(GrabBytes(NumBytes), NumBytes);
+	}
+
+	[[nodiscard]] FMemoryView GrabSkippableSlice()
+	{
+		return GrabSlice(GrabVarIntU());
+	}
+
+	[[nodiscard]] uint8 GrabByte()
+	{
+		return *GrabBytes(1);
+	}
+
+	template<typename T>
+	[[nodiscard]] T Grab()
+	{
+		return FPlatformMemory::ReadUnaligned<T>(GrabBytes(sizeof(T)));
+	}
+	
+	[[nodiscard]] uint64 GrabVarIntU()
+	{
+		check(It < End);
+		uint32 NumBytesRead;
+		uint64 Out = ReadVarUInt(It, NumBytesRead);
+		It += NumBytesRead;
+		check(It <= End);
+		return Out;
+	}
+
+	[[nodiscard]] const uint8* Peek() const { return It; }
+	
+	[[nodiscard]] FMemoryView PeekSkippableSlice() const { return FByteReader(*this).GrabSkippableSlice(); }
+
+	template<typename T>
+	void SkipAlignmentPadding()
+	{
+		for (; !IsAligned(It, alignof(T)); ++It)
+		{
+			check(*It == 0);
+		}
+	}
+	
+	void CheckEmpty() const { check(It == End); }
+	void CheckNonEmpty() const { check(It != End); }
+	void CheckSize(int64 ExpectedSize) const { check(End - It == ExpectedSize); }
+};
+
+// Helper that consumes 8 bits from the byte value stream
+class FBitCacheReader
+{
+	uint8 Bits = 0;
+	uint8 BitIt = 0;
+public:
+	[[nodiscard]] bool GrabNext(FByteReader& Bytes);
+};
+
+//////////////////////////////////////////////////////////////////////////
+
+struct FStructSchemaHandle
+{
+	FStructSchemaId		Id;
+	FReadBatchId		Batch;
+
+	const FStructSchema& Resolve() const { return ResolveStructSchema(Batch, Id); }
+	const FStructSchema& ResolveSuper() const;
+};
+
+struct FStructView
+{
+	FStructSchemaHandle	Schema;
+	FByteReader			Values;
+};
+
+//////////////////////////////////////////////////////////////////////////
+
+union FMemberValue
+{
+	const uint8*		Ptr;	// From byte stream
+	bool				bValue; // From bit cache
+};
+
+struct FLeafView
+{
+	ELeafType		Type;
+    ELeafWidth		Width;
+	FReadBatchId	Batch;
+	FEnumSchemaId	Enum;
+	FMemberValue	Value;
+
+	bool			AsBool() const		{ check(Type == ELeafType::Bool);	return Value.bValue; }
+	int8			AsS8() const		{ check(Type == ELeafType::IntS);	return As<int8>(); }
+	uint8			AsU8() const		{ check(Type == ELeafType::IntU);	return As<uint8>(); }
+	int16			AsS16() const		{ check(Type == ELeafType::IntS);	return As<int16>(); }
+	uint16			AsU16() const		{ check(Type == ELeafType::IntU);	return As<uint16>(); }
+	int32			AsS32() const		{ check(Type == ELeafType::IntS);	return As<int32>(); }
+	uint32			AsU32() const		{ check(Type == ELeafType::IntU);	return As<uint32>(); }
+	int64			AsS64() const		{ check(Type == ELeafType::IntS);	return As<int64>(); }
+	uint64			AsU64() const		{ check(Type == ELeafType::IntU);	return As<uint64>(); }
+	double			AsDouble() const	{ check(Type == ELeafType::Float);	return As<double>(); }
+	float			AsFloat() const		{ check(Type == ELeafType::Float);	return As<float>(); }
+	char8_t			AsChar8() const		{ check(Type == ELeafType::Unicode);return As<char8_t>(); }
+	char16_t		AsChar16() const	{ check(Type == ELeafType::Unicode);return As<char16_t>(); }
+	char32_t		AsChar32() const	{ check(Type == ELeafType::Unicode);return As<char32_t>(); }
+	uint8			AsEnum8() const		{ check(Type == ELeafType::Enum);	return As<uint8>(); }
+	uint16			AsEnum16() const	{ check(Type == ELeafType::Enum);	return As<uint16>(); }
+	uint32			AsEnum32() const	{ check(Type == ELeafType::Enum);	return As<uint32>(); }
+	uint64			AsEnum64() const	{ check(Type == ELeafType::Enum);	return As<uint64>(); }
+
+private:
+	template<typename T> T As() const
+	{
+		check(SizeOf(Width) == sizeof(T));
+		return *reinterpret_cast<const T*>(Value.Ptr);
+	}
+};
+
+//////////////////////////////////////////////////////////////////////////
+
+//struct FEnumSchemaFlatView
+//{
+//	const FEnumSchema& Schema;
+//	
+//	TOptional<uint64>		ToValue(FMemberId Id) const;
+//	TOptional<FMemberId>	ToId(uint64 Value) const;
+//	bool					IsSet(FMemberId Id, uint64 Value) const;
+//};
+//
+//struct FEnumSchemaFlagView
+//{
+//	const FEnumSchema& Schema;
+//
+//	FEnumSchemaFlagView(const FEnumSchema& InSchema); // @pre sizeof(T) matches Schema.Width and all schema values use single bit
+//	
+//	TOptional<uint64>						ToFlags(TConstArrayView<FMemberId> Ids) const;
+//	TOptional<TConstArrayView<FMemberId>>	ToIds(uint64 Flags) const;
+//	bool									HasAny(FMemberId Id, uint64 Flags) const;
+//	bool									HasAll(FMemberId Id, uint64 Flags) const;
+//};
+//
+//struct FEnumFlatView
+//{
+//	const FEnumSchemaFlatView* Schema;
+//	uint64 Value;
+//
+//	bool Equals(FMemberId Id) const;
+//};
+//
+//struct FEnumFlagView
+//{
+//	const FEnumSchemaFlagView* Schema;
+//	uint64 Value;
+//};
+
+//////////////////////////////////////////////////////////////////////////
+
+class FMemberReader;
+class FRangeLoader;
+class FStructRangeIterator;
+class FNestedRangeIterator;
+
+template<typename IteratorType>
+class TStructuralRangeView;
+
+class FLeafRangeView;
+using FStructRangeView = TStructuralRangeView<FStructRangeIterator>;
+using FNestedRangeView = TStructuralRangeView<FNestedRangeIterator>;
+
+struct FRangeSchema
+{
+	FMemberType				ItemType;
+	FReadBatchId			Batch; // Needed to resolve inner schema
+	FOptionalSchemaId		InnermostSchema;
+	const FMemberType*		NestedItemTypes; // For nested ranges, can be out-of-bounds otherwise
+};
+
+class FRangeView
+{
+public:
+	uint64						Num() const				{ return NumItems;}
+	bool						IsEmpty() const			{ return NumItems == 0; }
+	bool						IsLeafRange() const		{ return Schema.ItemType.GetKind() == EMemberKind::Leaf; }
+	bool						IsStructRange() const	{ return Schema.ItemType.GetKind() == EMemberKind::Struct; }
+	bool						IsNestedRange() const	{ return Schema.ItemType.GetKind() == EMemberKind::Range; }
+
+	PP_API FLeafRangeView		AsLeaves() const;		// @pre IsLeafRange()
+	PP_API FStructRangeView	AsStructs() const;		// @pre IsStructRange()
+	PP_API FNestedRangeView	AsRanges() const;		// @pre IsNestedRange()
+
+private:
+	friend FMemberReader;
+	friend FNestedRangeIterator;
+	friend FRangeLoader;
+
+	FRangeSchema			Schema;
+	uint64					NumItems;
+	FMemoryView				Values;
+};
+
+//////////////////////////////////////////////////////////////////////////
+
+class FBoolRangeIterator
+{
+public:
+	FBoolRangeIterator(const uint8* Data, uint64 Idx)
+	: Byte(Data + Idx / 8)
+	, Mask(1u << (Idx % 8))
+	{}
+
+	bool operator*() const
+	{
+		return !!((*Byte) & Mask);
+	}
+
+	FBoolRangeIterator& operator++()
+	{
+		Mask <<= 1;
+		if (Mask == 0x100)
+		{
+			++Byte;
+			Mask = 1;
+		}
+		return *this;
+	}
+
+	bool operator!=(FBoolRangeIterator Rhs)
+	{
+		return Byte != Rhs.Byte || Mask != Rhs.Mask;
+	}
+
+private:
+	const uint8* Byte;
+	uint32 Mask;
+};
+
+class FBoolRangeView
+{
+public:
+	FBoolRangeView(const uint8* InData, uint64 InNum)
+	: Data(InData)
+	, NumBits(InNum)
+	{}
+
+	uint64 Num() const					{ return NumBits;}
+	bool operator[](uint64 Idx) const	{ return *FBoolRangeIterator(Data, Idx); }
+	FBoolRangeIterator begin() const	{ return FBoolRangeIterator(Data, 0); }
+	FBoolRangeIterator end() const		{ return FBoolRangeIterator(Data, NumBits); }
+
+private:
+	const uint8* Data;
+	uint64 NumBits;
+
+	friend inline uint64 GetNum(const FBoolRangeView& Range) { return Range.NumBits; }
+};
+
+//////////////////////////////////////////////////////////////////////////
+
+template<typename T>
+class TRangeView
+{
+public:
+	TRangeView(const T* InData, uint64 InNum)
+	: Data(InData)
+	, NumItems(InNum)
+	{}
+
+	uint64 Num() const					{ return NumItems;}
+	T operator[](uint64 Idx) const		{ check(Idx < NumItems); return Data[Idx]; }
+	const T* begin() const				{ return Data; }
+	const T* end() const				{ return Data + NumItems; }
+private:
+	const T* Data;
+	uint64 NumItems;
+	
+	friend inline uint64 GetNum(const TRangeView& Range) { return Range.NumItems; }
+	friend inline const T* GetData(const TRangeView& Range) { return Range.Data; }
+};
+
+//  Works with FBoolRangeView, which lack GetData(), TRangeView, TArrayView, std::intializer_list and T[]
+template<class RangeTypeA, class RangeTypeB>
+bool EqualItems(RangeTypeA&& A, RangeTypeB&& B)
+{
+	if (static_cast<uint64>(GetNum(A)) != static_cast<uint64>(GetNum(B)))
+	{
+		return false;
+	}
+
+	auto BIt = std::begin(B);
+	for (auto bA : A)
+	{
+		if (bA != *BIt++)
+		{
+			return false;
+		}
+	}
+
+	return true;
+}
+
+//////////////////////////////////////////////////////////////////////////
+
+class FLeafRangeView
+{
+public:
+	ELeafType GetType() const;
+	SIZE_T GetSizeOf() const;
+	uint64 Num() const { return NumItems; }
+	uint64 NumBytes() const { return NumItems * GetSizeOf(); }
+	
+	// These ranges hide the internal representations so we can change format in the future, 
+	// e.g. store zeroes or 1.0f in some compact fashion or even var int encodings
+	//
+	// They could also provide various conversion helpers
+
+	FBoolRangeView			AsBools() const
+	{
+		check(Type == ELeafType::Bool);
+		return FBoolRangeView(Values, NumItems);
+	}
+	
+	TRangeView< int8>		AsS8s() const		{ return As< int8>(); }
+	TRangeView<uint8>		AsU8s() const		{ return As<uint8>(); }
+	TRangeView< int16>		AsS16s() const		{ return As< int16>(); }
+	TRangeView<uint16>		AsU16s() const		{ return As<uint16>(); }
+	TRangeView< int32>		AsS32s() const		{ return As< int32>(); }
+	TRangeView<uint32>		AsU32s() const		{ return As<uint32>(); }
+	TRangeView< int64>		AsS64s() const		{ return As< int64>(); }
+	TRangeView<uint64>		AsU64s() const		{ return As<uint64>(); }
+	TRangeView<float>		AsFloats() const	{ return As<float>(); }
+	TRangeView<double>		AsDoubles() const	{ return As<double>(); }
+	TRangeView<char8_t>		AsUtf8() const		{ return As<char8_t>(); }
+	TRangeView<char16_t>	AsUtf16() const		{ return As<char16_t>(); }
+	TRangeView<char32_t>	AsUtf32() const		{ return As<char32_t>(); }
+
+	template<typename LeafType>
+	TRangeView<LeafType> As() const
+	{
+		check(FUnpackedLeafType(Type, Width) == ReflectLeaf<LeafType>);
+		return TRangeView<LeafType>(reinterpret_cast<const LeafType*>(Values), NumItems);
+	}
+
+	//FEnumFlatRangeView		AsFlatEnums();	// @pre GetType() == ELeafType::Enum
+	//FEnumFlagRangeView		AsFlagEnums();	// @pre GetType() == ELeafType::Enum && 
+
+	//// Option B Memcopying + copy conversion apis
+	//
+	//void					ConvertBits(TArrayView<bool> Out);
+	//void					ConvertBits(TArrayView<uint8> Out);
+	//void					ConvertBits(TArrayView<uint16> Out);
+	//void					ConvertBits(TArrayView<uint32> Out);
+	//void					ConvertBits(TArrayView<uint64> Out);
+
+	//template<typename T>
+	//void					CopyExact(TArrayView<T> Out);
+
+	//// E.g. IntCastChecked, double -> float ok but not float -> double if it loses precision
+	////						uint64 -> uint32 ok if all 
+	//template<typename T>
+	//bool					TryConvertLossless(TArrayView<uint64> Out);
+
+	////  
+	//template<typename T>
+	//bool					TryConvert(TArrayView<T> Out);
+
+	//template<typename T>
+	//void					CopyStaticCast(TArrayView<T> Out);
+
+
+private:
+	friend class FRangeView;
+
+	ELeafType						Type;
+	ELeafWidth						Width;
+	FReadBatchId					Batch;
+	FOptionalEnumSchemaId			Enum;
+	uint64							NumItems;
+	const uint8*					Values;
+};
+
+//////////////////////////////////////////////////////////////////////////
+
+class FNestedRangeIterator
+{
+	friend FNestedRangeView;
+	FRangeSchema			Schema;
+	FByteReader				ByteIt;
+	FBitCacheReader			BitIt;
+
+public:
+	FNestedRangeIterator(const FRangeSchema& InSchema, FMemoryView Data)
+	: Schema(InSchema) 
+	, ByteIt(Data)
+	{}
+
+	PP_API FRangeView operator*() const;
+	bool operator!=(const FNestedRangeIterator& Rhs) const { return ByteIt.Peek() != Rhs.ByteIt.Peek(); };
+	PP_API void operator++();
+};
+
+class FStructRangeIterator
+{
+	friend FStructRangeView;
+	FStructSchemaHandle		Schema;
+	FByteReader				ByteIt;
+
+public:
+	FStructRangeIterator(const FStructSchemaHandle& InSchema, FMemoryView Data)
+	: Schema(InSchema) 
+	, ByteIt(Data)
+	{}
+
+	FStructView operator*() const { return { Schema, FByteReader(ByteIt.PeekSkippableSlice()) }; }
+	bool operator!=(const FStructRangeIterator& Rhs) const { return ByteIt.Peek() != Rhs.ByteIt.Peek(); };
+	void operator++() { (void)ByteIt.GrabSkippableSlice(); }
+};
+
+template<typename IteratorType>
+class TStructuralRangeView
+{
+public:
+	using SchemaType = decltype(IteratorType::Schema);
+
+	uint64			Num() const		{ return NumItems; }
+	SchemaType		GetSchema() const { return Schema; }
+
+	IteratorType	begin() const	{ return IteratorType(Schema, Data); }
+	IteratorType	end() const		{ return IteratorType(Schema, FMemoryView(Data.GetDataEnd(), 0)); }
+
+private:
+	friend class FRangeView;
+
+	uint64				NumItems;
+	FMemoryView			Data;
+	SchemaType			Schema;
+};
+
+//////////////////////////////////////////////////////////////////////////
+
+//class FAnyMemberView
+//{
+//public:
+//	EMemberKind				GetKind() const			{ return Type.Kind; };
+//
+//	bool					IsLeaf() const			{ return Type.Kind == EMemberKind::Leaf; }
+//	bool					IsStruct() const		{ return Type.Kind == EMemberKind::Struct; }
+//	bool					IsRange() const			{ return Type.Kind == EMemberKind::Range; }
+//
+//	FStructView				AsStruct() const;		// @pre bool() && IsStruct()
+//	FLeafView				AsLeaf() const;			// @pre bool() && IsLeaf()
+//	FRangeView				AsRange() const;		// @pre bool() && IsRange()
+//
+//private:
+//	friend class FMemberReader;
+//
+//	FMemberType			Type;
+//	uint32				BatchId; // To resolve struct schema
+//	FMemberValue		Value;
+//};
+
+// Iterates over struct members
+class FMemberReader
+{
+public:
+	explicit FMemberReader(FStructView Struct) : FMemberReader(Struct.Schema.Resolve(), Struct.Values, Struct.Schema.Batch) {}
+
+	bool					HasMore() const			{ return MemberIdx < NumMembers; }
+	
+	FOptionalMemberId		PeekName() const;		// @pre HasMore()
+	EMemberKind				PeekKind() const;		// @pre HasMore()
+	FMemberType				PeekType() const;		// @pre HasMore()
+
+	FLeafView				GrabLeaf();				// @pre PeekKind() == EMemberKind::Leaf
+	FRangeView				GrabRange();			// @pre PeekKind() == EMemberKind::Range
+	FStructView				GrabStruct();			// @pre PeekKind() == EMemberKind::Struct
+	//FAnyMemberView		GrabAny();				// @pre HasMore()
+
+protected: // for unit tests
+	const FMemberType*		Footer;
+	const FReadBatchId		Batch;					// Needed to resolve schemas
+	const uint16			NumMembers;
+	const uint16			NumRangeTypes;			// Number of ranges and nested ranges
+	const uint16			IsSparse : 1;
+	const uint16			HasSuper : 1;
+
+	uint16					MemberIdx = 0;
+	uint16					RangeTypeIdx = 0;		// Types of [nested] ranges
+	uint16					InnerSchemaIdx = 0;		// Types of static structs and enums
+	FBitCacheReader			Bits;
+	FByteReader				ValueIt;
+
+#if DO_CHECK
+	const uint16			NumInnerSchemas;		// Number of static structs and enums
+#endif
+
+	FMemberReader(const FStructSchema& Schema, FByteReader Values, FReadBatchId InBatch);
+
+	const FMemberType*		GetMemberTypes() const;
+	const FMemberType*		GetRangeTypes() const;
+	const FSchemaId*		GetInnerSchemas() const;
+	const FMemberId*		GetMemberNames() const;
+	
+	void					AdvanceToNextMember();
+	void					SkipMissingSparseMembers();
+	void					SkipSchema(FMemberType InnermostType);
+
+	using FMemberTypeRange = TConstArrayView<FMemberType>;
+
+	FMemberTypeRange		GrabRangeTypes();
+	FSchemaId				GrabInnerSchema();
+	FStructSchemaId			GrabStructSchema(FStructType Type);
+	FOptionalSchemaId		GrabRangeSchema(FMemberType InnermostType);
+	FEnumSchemaId			GrabEnumSchema()			{ return static_cast<FEnumSchemaId&&>(GrabInnerSchema()); }
+	bool					GrabBit()					{ return Bits.GrabNext(/* in-out */ ValueIt); }
+	uint64					GrabSkipLength()			{ return ValueIt.GrabVarIntU(); }
+};
+
+//////////////////////////////////////////////////////////////////////////
+
+// Hides the inheritance chain and iterates over super members first
+class FFlatMemberReader
+{
+public:
+	FFlatMemberReader(FStructView Struct);
+
+	bool					HasMore() const			{ return It->HasMore(); }
+
+	FMemberId				PeekName() const		{ return It->PeekName().Get(); }
+	EMemberKind				PeekKind() const		{ return It->PeekKind(); }
+	FTypeId					PeekOwner() const		{ return It->Owner; }
+
+	FLeafView				GrabLeaf()				{ return Grabbed(It->GrabLeaf()); }
+	FRangeView				GrabRange()				{ return Grabbed(It->GrabRange()); }
+	FStructView				GrabStruct()			{ return Grabbed(It->GrabStruct()); }
+
+private:
+	struct FReader : public FMemberReader
+	{
+		FReader(FStructView Struct);
+		FTypeId Owner;
+	};
+	TArray<FReader, TInlineAllocator<8>> Lineage;
+	FReader* It = nullptr;
+
+	template<class ViewType>
+	ViewType Grabbed(ViewType Out)
+	{
+		It -= (It != &Lineage[0] && !It->HasMore());
+		return Out;
+	}
+};
+
+} // namespace PlainProps
\ No newline at end of file
diff --git a/Engine/Source/Programs/PlainPropsTests/Public/PlainPropsSave.h b/Engine/Source/Programs/PlainPropsTests/Public/PlainPropsSave.h
new file mode 100644
--- /dev/null
+++ b/Engine/Source/Programs/PlainPropsTests/Public/PlainPropsSave.h
@@ -0,0 +1,25 @@
+// Copyright Epic Games, Inc. All Rights Reserved.
+
+#pragma once
+
+#include  "PlainPropsTypes.h"
+#include "Templates/UniquePtr.h"
+
+namespace PlainProps 
+{
+
+struct FBuiltStruct;
+class FDeclarations;
+class FStructBindings;
+
+struct FSaveContext
+{
+	const FDeclarations& Declarations;
+	const FStructBindings& Bindings;
+	const FDebugIds& Debug;
+};
+
+[[nodiscard]] TUniquePtr<FBuiltStruct> SaveStruct(const uint8* Struct, FStructSchemaId Id, const FSaveContext& Context);
+
+} // namespace PlainProps
+
diff --git a/Engine/Source/Programs/PlainPropsTests/Public/PlainPropsTypes.h b/Engine/Source/Programs/PlainPropsTests/Public/PlainPropsTypes.h
new file mode 100644
--- /dev/null
+++ b/Engine/Source/Programs/PlainPropsTests/Public/PlainPropsTypes.h
@@ -0,0 +1,379 @@
+// Copyright Epic Games, Inc. All Rights Reserved.
+
+#pragma once
+
+#include "Containers/ArrayView.h"
+#include "Misc/AssertionMacros.h"
+#include "Templates/TypeHash.h"
+#include <type_traits>
+
+#define PP_API
+
+namespace PlainProps
+{
+
+enum class EMemberKind : uint8 { Leaf, Struct, Range };
+enum class ELeafType : uint8 { Bool, IntS, IntU, Float, Hex, Enum, Unicode };
+enum class ELeafWidth : uint8 { B8, B16, B32, B64 };
+enum class ERangeSizeType : uint8 { Uni, S8, U8, S16, U16, S32, U32, S64, U64 };
+
+inline constexpr SIZE_T SizeOf(ELeafWidth Width)
+{
+	return SIZE_T(1) << static_cast<uint32>(Width);
+}
+
+inline SIZE_T SizeOf(ERangeSizeType Width)
+{
+	check(Width != ERangeSizeType::Uni);
+    return SIZE_T(1) << ((uint8)Width - 1) / 2; 
+}
+
+inline constexpr uint64 Max(ERangeSizeType Width)
+{
+    const uint8_t LeadingZeroes[] = {63, 57, 56, 49, 48, 33, 32, 1, 0};
+    return ~uint64_t(0) >> LeadingZeroes[(uint8)Width]; 
+}
+
+template<SIZE_T Size>
+constexpr ELeafWidth IllegalLeafWidth()
+{
+	static_assert(Size == 1);
+	return ELeafWidth::B8;
+}
+
+template<SIZE_T Size> constexpr ELeafWidth LeafWidth = IllegalLeafWidth<Size>();
+template<> inline constexpr ELeafWidth LeafWidth<1> = ELeafWidth::B8;
+template<> inline constexpr ELeafWidth LeafWidth<2> = ELeafWidth::B16;
+template<> inline constexpr ELeafWidth LeafWidth<4> = ELeafWidth::B32;
+template<> inline constexpr ELeafWidth LeafWidth<8> = ELeafWidth::B64;
+
+//////////////////////////////////////////////////////////////////////////
+
+struct FLeafType
+{
+    EMemberKind		_ : 2;
+    ELeafWidth		Width : 2;
+    ELeafType		Type : 3;
+};
+
+struct FRangeType
+{
+    EMemberKind		_ : 2;
+    ERangeSizeType	MaxSize : 4;
+};
+
+struct FStructType
+{
+    EMemberKind		_ : 2;
+	uint8			IsDynamic : 1;	// Which schema stored in value stream, different instances in same batch have different types
+	uint8			IsSuper : 1;	// Inherited members from some super struct, can only appear as first member
+//	uint8			IsDense : 1;	// Potential optimization to avoid accessing schema
+};
+
+union FMemberType
+{
+	FMemberType() = default; // Uninitialized
+	constexpr explicit FMemberType(FLeafType InLeaf) : Leaf(InLeaf) {}
+	constexpr explicit FMemberType(ELeafType Type, ELeafWidth Width) : Leaf({EMemberKind::Leaf, Width, Type}) {}
+	constexpr explicit FMemberType(FRangeType InRange) : Range(InRange) {}
+	constexpr explicit FMemberType(ERangeSizeType MaxSize) : Range({EMemberKind::Range, MaxSize}) {}
+	constexpr explicit FMemberType(FStructType InStruct) : Struct(InStruct) {}
+
+	bool			IsLeaf() const		{ return Kind == EMemberKind::Leaf; }
+	bool			IsRange() const		{ return Kind == EMemberKind::Range; }
+	bool			IsStruct() const	{ return Kind == EMemberKind::Struct; }
+	EMemberKind     GetKind() const		{ return Kind; }
+
+	FLeafType		AsLeaf() const		{ check(IsLeaf());		return Leaf; }
+	FRangeType		AsRange() const		{ check(IsRange());		return Range; }
+	FStructType		AsStruct() const	{ check(IsStruct());	return Struct; }
+	uint8			AsByte() const		{ return reinterpret_cast<const uint8&>(*this); }
+
+	friend inline bool operator==(FMemberType A, FMemberType B) { return A.AsByte() == B.AsByte(); }
+private:
+    EMemberKind     Kind : 2;
+    FLeafType	    Leaf;
+    FRangeType		Range;
+    FStructType		Struct;
+};
+
+//////////////////////////////////////////////////////////////////////////
+
+struct FNameId
+{
+	uint32 Idx = ~0u;
+
+	bool operator==(FNameId O) const { return Idx == O.Idx; }
+	friend uint32 GetTypeHash(FNameId Id) { return Id.Idx; };
+};
+
+struct FMemberId
+{
+	FNameId Id;
+
+	bool operator==(FMemberId O) const { return Id == O.Id; }
+	friend uint32 GetTypeHash(FMemberId Member) { return Member.Id.Idx; };
+};
+
+//////////////////////////////////////////////////////////////////////////
+
+struct FSchemaId
+{
+	uint32 Idx = ~0u;
+	
+	// TODO: Consider moving these to FEnumSchemaId/FStructSchemaId.
+	// They're not comparable in the building stage when the overlap.
+	// After remap they are comparable during reading.
+	bool operator==(FSchemaId O) const { return Idx == O.Idx; }
+	friend uint32 GetTypeHash(FSchemaId Schema) { return GetTypeHash(Schema.Idx); };
+};
+struct FStructSchemaId : FSchemaId {};
+struct FEnumSchemaId : FSchemaId {};
+
+//////////////////////////////////////////////////////////////////////////
+
+struct FNestedScopeId { uint32 Idx; };
+struct FFlatScopeId { FNameId Name; };
+
+class FScopeId
+{
+	static constexpr uint32 NestedBit = 0x80000000u;
+	uint32 Handle;
+public:
+	explicit FScopeId(FFlatScopeId Flat) : Handle(Flat.Name.Idx) 				{ check(AsFlat().Name == Flat.Name); }
+	explicit FScopeId(FNestedScopeId Nested) : Handle(Nested.Idx | NestedBit) 	{ check(AsNested().Idx == Nested.Idx); }
+
+	bool							IsFlat() const 					{ return  !(Handle & NestedBit); }
+	bool							IsNested() const 				{ return !!(Handle & NestedBit); }
+	FFlatScopeId					AsFlat() const 					{ check(IsFlat()); 		return {FNameId{Handle}}; }
+	FNestedScopeId					AsNested() const				{ check(IsNested()); 	return {Handle & ~NestedBit}; }
+	uint32							AsInt() const					{ return Handle; }
+	bool							operator==(FScopeId O) const	{ return Handle == O.Handle; }
+};
+
+//////////////////////////////////////////////////////////////////////////
+
+struct FConcreteTypenameId
+{
+	FNameId Id;
+	bool operator==(FConcreteTypenameId O) const { return Id == O.Id; }
+};
+
+struct FBaseTypenameId
+{
+	FBaseTypenameId(uint8 InNumParameters, uint32 InIdx)
+	: NumParameters(InNumParameters)
+	, Idx(InIdx)
+	{
+		check(Idx == InIdx);
+	}
+
+	uint32 NumParameters : 8;
+	uint32 Idx : 24;
+
+	uint32 AsInt() const { return (Idx << 8) + NumParameters; }
+	static FBaseTypenameId FromInt(uint32 Int) { return FBaseTypenameId(static_cast<uint8>(Int), Int >> 8); }
+	bool operator==(FBaseTypenameId O) const { return AsInt() == O.AsInt(); }
+};
+
+struct FParametricTypeId : FBaseTypenameId
+{
+	using FBaseTypenameId::FBaseTypenameId;
+	static FParametricTypeId FromInt(uint32 Int) { return static_cast<FParametricTypeId&&>(FBaseTypenameId::FromInt(Int)); }
+};
+
+struct FTypenameId : FBaseTypenameId
+{
+	explicit FTypenameId(FParametricTypeId Parametric) : FBaseTypenameId(Parametric) { check(AsParametric().AsInt() == Parametric.AsInt()); }
+	explicit FTypenameId(FConcreteTypenameId Concrete) : FBaseTypenameId(0, Concrete.Id.Idx) { check(AsConcrete().Id == Concrete.Id); }
+
+	bool							IsConcrete() const 		{ return NumParameters == 0; }
+	bool							IsParametric() const 	{ return !IsConcrete(); }
+	FConcreteTypenameId				AsConcrete() const 		{ check(IsConcrete());		return { FNameId{Idx} };  }
+	FParametricTypeId				AsParametric() const 	{ check(IsParametric());	return FParametricTypeId(NumParameters, Idx); }
+};
+
+//////////////////////////////////////////////////////////////////////////
+
+struct FTypeId
+{
+	FScopeId 		Scope;
+	FTypenameId		Name;
+
+	bool operator==(FTypeId O) const { return Scope.AsInt() == O.Scope.AsInt() && Name.AsInt() == O.Name.AsInt(); }
+	friend uint32 GetTypeHash(FTypeId Type) { return HashCombineFast(Type.Scope.AsInt(), Type.Name.Idx); };
+};
+
+//////////////////////////////////////////////////////////////////////////
+
+struct FNoId {};
+inline constexpr FNoId NoId;
+
+
+inline uint32 ToIdx(FNameId Id) { return Id.Idx; }
+inline uint32 ToIdx(FMemberId Name) { return Name.Id.Idx; }
+inline uint32 ToIdx(FSchemaId Id) { return Id.Idx; }
+inline uint32 ToIdx(FNestedScopeId Id) { return Id.Idx; }
+inline uint32 ToIdx(FParametricTypeId Id) { return Id.AsInt(); }
+
+template<class IdType>
+IdType FromIdx(uint32 Idx) { return {Idx}; }
+
+template<>
+inline FMemberId FromIdx(uint32 Idx) { return {{Idx}}; }
+
+template<>
+inline FParametricTypeId FromIdx(uint32 Idx) { return FParametricTypeId::FromInt(Idx); }
+
+class FOptionalId
+{
+protected:
+    uint32 Idx = ~0u; 
+public:
+    explicit operator bool() const { return Idx != ~0u; }
+	friend uint32 GetTypeHash(FOptionalId Id) { return Id.Idx; };
+};
+
+template<class T>
+class TOptionalId : public FOptionalId
+{
+public:
+	constexpr TOptionalId() = default;
+	constexpr TOptionalId(FNoId) {}
+    constexpr TOptionalId(T Id) { Idx = ToIdx(Id); }
+    
+	template<class U>
+    explicit operator TOptionalId<U>() const
+    {
+		static_assert(std::is_convertible_v<U, T> || std::is_convertible_v<T, U>);
+		return static_cast<const TOptionalId<U>&>(static_cast<const FOptionalId&>(*this));
+    }
+
+    T Get() const
+	{
+		check(*this);
+		return FromIdx<T>(Idx);
+	}
+
+	friend bool operator==(TOptionalId A, TOptionalId B) { return A.Idx == B.Idx;}
+	friend bool operator!=(TOptionalId A, TOptionalId B) { return A.Idx != B.Idx;}
+};
+
+
+using FOptionalNameId = TOptionalId<FNameId>;
+using FOptionalMemberId = TOptionalId<FMemberId>;
+using FOptionalSchemaId = TOptionalId<FSchemaId>;
+using FOptionalStructSchemaId = TOptionalId<FStructSchemaId>;
+using FOptionalEnumSchemaId = TOptionalId<FEnumSchemaId>;
+using FOptionalNestedScopeId = TOptionalId<FNestedScopeId>;
+using FOptionalParametricTypeId = TOptionalId<FParametricTypeId>;
+
+template<class IdType>
+inline constexpr TOptionalId<IdType> ToOptional(IdType Id) { return Id; }
+
+//////////////////////////////////////////////////////////////////////////
+
+// Resolved FNestedScopeId
+struct FNestedScope
+{
+	FScopeId		Outer;
+	FFlatScopeId 	Inner;
+
+	bool operator==(FNestedScope O) const { return Outer.AsInt() == O.Outer.AsInt() && Inner.Name == O.Inner.Name; }
+	friend uint32 GetTypeHash(FNestedScope Scope) { return HashCombineFast(Scope.Outer.AsInt(), Scope.Inner.Name.Idx); };
+};
+
+
+struct FParameterIndexRange : FBaseTypenameId { using FBaseTypenameId::FBaseTypenameId; };
+
+// Name-resolved FParametricTypeId
+struct FParametricType
+{
+	FConcreteTypenameId		Name;
+	FParameterIndexRange	Parameters;
+
+	bool operator==(FParametricType O) const { return Name == O.Name && Parameters.AsInt() == O.Parameters.AsInt(); }
+};
+
+// Fully resolved FParametricTypeId
+struct FParametricTypeView
+{
+	FParametricTypeView(FConcreteTypenameId InName, uint8 NumParams, const FTypeId* Params) : Name(InName), NumParameters(NumParams), Parameters(Params) {}
+	FParametricTypeView(FConcreteTypenameId InName, TConstArrayView<FTypeId> Params);
+
+	FConcreteTypenameId			Name;
+	uint8						NumParameters;
+	const FTypeId*				Parameters;
+
+	TConstArrayView<FTypeId>	GetParameters() const { return MakeArrayView(Parameters, NumParameters); }
+};
+
+//////////////////////////////////////////////////////////////////////////
+
+struct FUnpackedLeafType
+{
+	ELeafType Type;
+	ELeafWidth Width;
+
+	constexpr FUnpackedLeafType(ELeafType InType, ELeafWidth InWidth) : Type(InType), Width(InWidth) {}
+	constexpr FUnpackedLeafType(FLeafType In) : Type(In.Type), Width(In.Width) {}
+
+	constexpr bool operator==(FUnpackedLeafType O) { return Type == O.Type && Width == O.Width; }
+	FMemberType Pack() const { return FMemberType(Type, Width); }
+};
+
+template<typename T>
+constexpr FUnpackedLeafType ReflectEnum()
+{
+	static_assert(std::is_enum_v<T>);
+	return FUnpackedLeafType(ELeafType::Enum, LeafWidth<sizeof(T)>);
+}
+
+template<typename T>
+constexpr FUnpackedLeafType ReflectLeaf = ReflectEnum<T>();
+
+template<> inline constexpr FUnpackedLeafType ReflectLeaf<bool>		= { ELeafType::Bool,	ELeafWidth::B8 };
+template<> inline constexpr FUnpackedLeafType ReflectLeaf<int8>		= { ELeafType::IntS,	ELeafWidth::B8 };
+template<> inline constexpr FUnpackedLeafType ReflectLeaf<int16>	= { ELeafType::IntS,	ELeafWidth::B16 };
+template<> inline constexpr FUnpackedLeafType ReflectLeaf<int32>	= { ELeafType::IntS,	ELeafWidth::B32 };
+template<> inline constexpr FUnpackedLeafType ReflectLeaf<int64>	= { ELeafType::IntS,	ELeafWidth::B64 };
+template<> inline constexpr FUnpackedLeafType ReflectLeaf<uint8>	= { ELeafType::IntU,	ELeafWidth::B8 };
+template<> inline constexpr FUnpackedLeafType ReflectLeaf<uint16>	= { ELeafType::IntU,	ELeafWidth::B16 };
+template<> inline constexpr FUnpackedLeafType ReflectLeaf<uint32>	= { ELeafType::IntU,	ELeafWidth::B32 };
+template<> inline constexpr FUnpackedLeafType ReflectLeaf<uint64>	= { ELeafType::IntU,	ELeafWidth::B64 };
+template<> inline constexpr FUnpackedLeafType ReflectLeaf<float>	= { ELeafType::Float,	ELeafWidth::B32 };
+template<> inline constexpr FUnpackedLeafType ReflectLeaf<double>	= { ELeafType::Float,	ELeafWidth::B64 };
+template<> inline constexpr FUnpackedLeafType ReflectLeaf<char8_t>	= { ELeafType::Unicode,	ELeafWidth::B8 };
+template<> inline constexpr FUnpackedLeafType ReflectLeaf<char16_t>	= { ELeafType::Unicode,	ELeafWidth::B16 };
+template<> inline constexpr FUnpackedLeafType ReflectLeaf<char32_t>	= { ELeafType::Unicode,	ELeafWidth::B32 };
+
+//////////////////////////////////////////////////////////////////////////
+
+template<typename T>
+const T* AlignPtr(const void* Ptr)
+{
+	return static_cast<const T*>(Align(Ptr, alignof(T)));
+}
+
+//////////////////////////////////////////////////////////////////////////
+
+class FDebugIds
+{
+public:
+	virtual ~FDebugIds() {}
+
+	virtual FParametricTypeView			Resolve(FParametricTypeId Id) const = 0;
+	virtual FNestedScope				Resolve(FNestedScopeId Id) const = 0;
+	virtual void						AppendDebugString(FString& Out, FNameId Name) const = 0;
+
+	virtual void						AppendDebugString(FString& Out, FScopeId Scope) const;
+	virtual void						AppendDebugString(FString& Out, FTypenameId Typename) const;
+	virtual void						AppendDebugString(FString& Out, FTypeId Type) const;
+
+	PP_API FString					Print(FNameId Name) const;
+	PP_API FString					Print(FMemberId Name) const;
+	PP_API FString					Print(FOptionalMemberId Name) const;
+	PP_API FString					Print(FTypeId Type) const;
+};
+
+} // namespace PlainProps
\ No newline at end of file
diff --git a/Engine/Source/Programs/PlainPropsTests/Public/PlainPropsUeCoreBindings.h b/Engine/Source/Programs/PlainPropsTests/Public/PlainPropsUeCoreBindings.h
new file mode 100644
--- /dev/null
+++ b/Engine/Source/Programs/PlainPropsTests/Public/PlainPropsUeCoreBindings.h
@@ -0,0 +1,444 @@
+// Copyright Epic Games, Inc. All Rights Reserved.
+
+#pragma once
+
+#include "Containers/Set.h"
+#include  "PlainPropsBind.h"
+#include  "PlainPropsLoad.h"
+#include  "PlainPropsRead.h"
+#include  "PlainPropsIndex.h"
+#include "UObject/NameTypes.h"
+
+namespace PlainProps
+{
+	template<>
+	void AppendString(FString& Out, const FName& Name) { Name.AppendString(Out); }
+}
+
+namespace PlainProps::UE
+{
+
+//class FReflection
+//{
+//	TIdIndexer<FName>	Names;
+//	FRuntime			Types;
+//
+//public:
+//	TIdIndexer<FName>&		GetIds() { return Names; }
+
+	//template<typename Ctti>
+	//FStructSchemaId			BindStruct();
+
+	//template<typename Ctti>
+	//FStructSchemaId			BindStructInterlaced(TConstArrayView<FMemberBinding> NonCttiMembers);
+	//FStructSchemaId			BindStruct(FStructSchemaId Id, const ICustomStructBinding& Custom);
+	//FStructSchemaId			BindStruct(FTypeId Type, FOptionalSchemaId Super, TConstArrayView<FNamedMemberBinding> Members, EMemberPresence Occupancy);
+	//void					DropStruct(FStructSchemaId Id) { Types.DropStruct(Id); }
+
+	//template<typename Ctti>
+	//FEnumSchemaId			BindEnum();
+	//FEnumSchemaId			BindEnum(FTypeId Type, EEnumMode Mode, ELeafWidth Width, TConstArrayView<FEnumerator> Enumerators);
+	//void					DropEnum(FEnumSchemaId Id) { Types.DropEnum(Id); }
+//};
+//
+//PP_API FReflection GReflection;
+//
+//struct FIds
+//{
+//	static FMemberId		IndexMember(FAnsiStringView Name)			{ return GReflection.GetIds().NameMember(FName(Name)); }
+//	static FTypenameId		IndexTypename(FAnsiStringView Name)			{ return GReflection.GetIds().MakeTypename(FName(Name)); }
+//	static FScopeId			IndexCurrentModule()						{ return GReflection.GetIds().MakeScope(FName(UE_MODULE_NAME)); }
+//	static FTypeId			IndexNativeType(FAnsiStringView Typename)	{ return {IndexCurrentModule(), IndexTypename(Typename)}; }
+//	static FEnumSchemaId	IndexEnum(FTypeId Name)						;
+//	static FEnumSchemaId	IndexEnum(FAnsiStringView Name)				;
+//	static FStructSchemaId	IndexStruct(FTypeId Name)					;
+//	static FStructSchemaId	IndexStruct(FAnsiStringView Name)			;
+//	
+//};
+
+// todo: use generic cached instance template?
+//template<class Ids>
+//FScopeId GetModuleScope()
+//{
+//	static FScopeId Id = Ids::IndexScope(UE_MODULE_NAME);
+//	return Id;
+//}
+
+//template<class Ctti>
+//class TBindRtti
+//{
+//	FSchemaId Id;
+//public:
+//	TBindRtti() : Id(BindRtti<Ctti, FIds>(GReflection.GetTypes()))
+//	{}
+//
+//	~TBindRtti()
+//	{
+//		if constexpr (std::is_enum_v<Ctti::Type>)
+//		{
+//			GReflection.DropEnum(static_cast<FEnumSchemaId>(Id));
+//		}
+//		else
+//		{
+//			GReflection.DropStruct(static_cast<FStructSchemaId>(Id));
+//		}
+//	}
+//};
+
+} // namespace PlainProps::UE
+
+//#define UEPP_BIND_STRUCT(T) 
+	
+//////////////////////////////////////////////////////////////////////////
+// Below container bindings should be moved to some suitable header
+//////////////////////////////////////////////////////////////////////////
+
+#include "Containers/Array.h"
+#include "Containers/Set.h"
+
+namespace PlainProps::UE
+{
+
+template <typename T>
+struct TArrayBinding : public IRangeBinding
+{
+	using SizeType = int32;
+	using ItemType = T;
+
+	virtual void MakeItems(FLoadRangeContext& Ctx) const override
+	{
+		TArray<T>& Array = *Ctx.Request.GetRange<TArray<T>>();
+		if constexpr (std::is_default_constructible_v<T>)
+		{
+			Array.SetNum(Ctx.Request.NumTotal());
+		}
+		else
+		{
+			Array.SetNumUninitialized(Ctx.Request.NumTotal());
+			Ctx.Items.SetUnconstructed();
+		}
+		
+		Ctx.Items.Set(Array.GetData(), Ctx.Request.NumTotal());
+	}
+
+	virtual void ReadItems(FSaveRangeContext& Ctx) const override
+	{
+		const TArray<T>& Array = *Ctx.Request.GetRange<TArray<T>>();
+		Ctx.Items.SetAll(Array.GetData(), static_cast<uint64>(Array.Num()));
+	}
+};
+
+//////////////////////////////////////////////////////////////////////////
+
+template <typename T>
+struct TSetBinding : public IRangeBinding
+{
+	using SizeType = int32;
+
+	virtual void MakeItems(FLoadRangeContext& Ctx) const override
+	{
+		const TSet<T>& Set = *Ctx.Request.GetRange<TSet<T>>();
+		SizeType Num = static_cast<SizeType>(Ctx.Request.NumTotal());
+
+		static constexpr bool bAllocate = sizeof(T) > sizeof(FLoadRangeContext::Scratch);
+		static constexpr uint64 MaxItems = bAllocate ? sizeof(FLoadRangeContext::Scratch) / sizeof(T) : 1;
+		
+		if (Ctx.Request.IsFirstCall())
+		{
+			Set.Reset();
+
+			if (uint64 NumRequested = Ctx.Request.NumTotal())
+			{
+				Set.Reserve(NumRequested);
+
+				// Create temporary buffer
+				uint64 NumTmp = FMath::Min(MaxItems, NumRequested);
+				void* Tmp = bAllocate ? FMemory::Malloc(sizeof(T)) : Ctx.Scratch;
+				Ctx.Items.Set(Tmp, NumTmp, sizeof(T));
+				if constexpr (std::is_default_constructible_v<T>)
+				{
+					for (T* It = static_cast<T*>(Tmp), End = It + NumTmp; It != End; ++It)
+					{
+						::new (It) T;
+					}
+				}
+				else
+				{
+					Ctx.Items.SetUnconstructed();
+				}
+
+				Ctx.Items.RequestFinalCall();
+			}
+		}
+		else
+		{
+			// Add items that have been loaded
+			TArrayView<T> Tmp = Ctx.Items.Get<T>();
+			for (T& Item : Tmp)
+			{
+				Set.Emplace(MoveTemp(Item));
+			}
+
+			if (Ctx.Request.IsFinalCall())
+			{
+				// Destroy and free temporaries
+				uint64 NumTmp = FMath::Min(MaxItems, Ctx.Request.NumTotal());
+				for (T& Item : MakeArrayView(Tmp.GetData(), NumTmp))
+				{
+					Item.~T();
+				}
+				if constexpr (bAllocate)
+				{
+					FMemory::Free(Tmp.GetData());
+				}	
+			}
+			else
+			{
+				Ctx.Items.Set(Tmp.Slice(0, FMath::Min(static_cast<uint64>(Tmp.Num()), Ctx.Request.NumMore())));
+				check(Ctx.Items.Get<T>().Num());
+			}
+		}
+	}
+
+	virtual void ReadItems(FSaveRangeContext& Ctx) const override
+	{
+		static_assert(offsetof(TSetElement<T>, Value) == 0);
+		const TSparseArray<TSetElement<T>>& Elems = *Ctx.Request.GetRange<TSparseArray<TSetElement<T>>>();
+
+		if (int32 NumRead = static_cast<int32>(Ctx.Items.NumTotal))
+		{
+			// Continue partial response
+			const TSetElement<T>* NextElem = static_cast<const TSetElement<T>*>(Ctx.Items.Slice) + Ctx.Items.Slice.Num + /* skip known invalid */ 1;
+			Ctx.Items.SetPart(GetContiguousSlice(Elems.PointerToIndex(NextElem), Elems));
+		}
+		else if (Elems.IsCompact())
+		{
+			Ctx.Items.SetAll(Elems.GetData(), Elems.Num());
+		}
+		else
+		{
+			// Start partial response
+			Ctx.Items.Stride = sizeof(TSetElement<T>);
+			Ctx.Items.SetPart(GetContiguousSlice(0, Elems));
+		}
+	}
+
+	FExistingItemSlice GetContiguousSlice(int32 Idx, const TSparseArray<TSetElement<T>>& Elems)
+	{
+		int32 Num = 1;
+		for (;!Elems.IsValidIndex(Idx); ++Idx) {}
+		for (; Elems.IsValidIndex(Idx + Num); ++Num) {}
+		return { &Elems[Idx], Num };
+	}
+};
+
+
+//////////////////////////////////////////////////////////////////////////
+
+struct FSetOps
+{
+	union
+	{
+		struct
+		{
+			FMemberId Add;
+			FMemberId Del;
+			FMemberId Set;
+		};
+
+		FMemberId All[3];
+	};
+	
+	template<class Ids>	
+	static const FSetOps& Get()
+	{
+		static FSetOps Ops = {Ids::IndexMember("Add"), Ids::IndexMember("Del"), Ids::IndexMember("Set")};
+		return Ops;
+	}
+};
+
+template <class Ids, typename T>
+struct TSetDeltaBinding : public ICustomStructBinding
+{
+	using Type = TSet<T>;
+	static constexpr EMemberPresence Occupancy = EMemberPresence::AllowSparse;
+	
+	static TConstArrayView<FMemberId> GetMemberIds() { return FSetOps::Get<Ids>().All; }
+
+	virtual void SaveStruct(FMemberBuilder& Dst, const void* Src, void* UserData, const FDebugIds& Debug) const override;
+
+	virtual void LoadStruct(void* Dst, FStructView Src, ECustomLoadMethod Method, const FLoadBatch& Batch) const override
+	{
+		Type& Out = *reinterpret_cast<Type*>(Dst);
+		FMemberReader Members(Src);
+
+		if (Method == ECustomLoadMethod::Construct)
+		{
+			::new (Dst) Type;
+		}
+				
+		if (!Members.HasMore())
+		{
+			return;
+		}
+
+		FSetOps Ops = FSetOps::Get<Ids>();
+		FMemberId Name = Members.PeekName().Get();
+		FRangeView Items = Members.GrabRange();
+		int32 NumItems = static_cast<int32>(Items.Num());
+		if (Name == Ops.Set)
+		{
+			Out.Empty(NumItems);
+			AddItems(Out, Items);
+		}
+		else if (Name == Ops.Add)
+		{
+			Out.Reserve(Out.Num() + NumItems);
+			AddItems(Out, Items);
+		}
+		else if (Name == Ops.Del)
+		{
+			RemoveItems(Out, NumItems);
+		
+			if (Members.HasMore())
+			{
+				check(Members.PeekName() == Ops.Add);
+				Items = Members.GrabRange();
+				Out.Reserve(Out.Num() + static_cast<int32>(Items.Num()));
+				AddItems(Out, Items);
+			}
+		}
+		
+		check(!Members.HasMore());
+	}
+
+	void AddItems(Type& Out, FRangeView Items, const FLoadBatch& Batch)
+	{
+		check(!Items.IsEmpty());
+
+		if constexpr (std::is_arithmetic_v<T> || std::is_enum_v<T>)
+		{
+			for (T Item : Items.AsLeaves().As<T>())
+			{
+				Out.Add(Item);
+			}
+		}
+		else if (Items.IsStructRange())
+		{
+			FStructRangeView Structs = Items.AsStructs();
+			for (FStructView Item : Structs)
+			{
+				if constexpr (std::is_default_constructible_v<T>())
+				{
+					T Tmp;
+					PlainProps::LoadStruct(&Tmp, Item, Batch);	
+					Out.Emplace(MoveTemp(Item));
+				}
+				else
+				{
+					alignas(T) uint8 Buffer[sizeof(T)];
+					PlainProps::ConstructAndLoadStruct(Buffer, Item, Batch);
+					T* Tmp = reinterpret_cast<T*>(Buffer);
+					Out.Emplace(MoveTemp(*Tmp));
+					Tmp->~T();
+				}
+			}
+		}
+		else if constexpr (std::is_default_constructible_v<T>())
+		{
+			using Binding = RangeBind<T>;
+			const IRangeBinding* Bindings[] = {};// ... generate somehow ... };
+			T Tmp;
+			for (FRangeView Item : Items.AsRanges())
+			{
+				LoadRange(&Tmp, Item, Bindings, Batch);
+				Out.Emplace(MoveTemp(*Tmp));
+			}
+		}
+		else
+		{
+			check(Items.IsNestedRange());
+			checkf(false, TEXT("Ranges must be default-constructible"));
+		}
+	}
+};
+
+//	template <typename T, typename VariantType>
+//	struct TVariantConstructFromMember
+//	{
+//		/** Default construct the type and load it from the FArchive */
+//		static void Construct(VariantType& Dst, FMemberReader Src)
+//		{
+//			if constexpr (std::is_arithmetic_v<T>)
+//			{
+//				Dst.template Emplace<T>(Src.GrabLeaf().As<T>());
+//			}
+//			else constexpr 
+//
+//		}
+//	};
+//
+//template <typename... Ts>
+//struct TVariantBinding : public ICustomStructBinding
+//{
+//	using VariantType = TVariant<Ts...>;
+//
+//	const FStructDeclaration& Declaration;
+//
+//	static constexpr void(*Loaders[])(FMemberReader&, VariantType&) = { &TLoader<Ts, VariantType>::Load... };
+//
+//	virtual void LoadStruct(void* Dst, FStructView Src, ECustomLoadMethod Method, const FLoadBatch& Batch) const override
+//	{
+//		VariantType& Variant = *reinterpret_cast<VariantType*>(Dst);
+//
+//		if (Method == ECustomLoadMethod::Assign)
+//		{
+//			Variant.~VariantType();
+//		}
+//
+//		FMemberReader Members(Src);
+//		const FMemberId* DeclaredName = Algo::Find(Declaration.Names, Members.PeekName());
+//		check(DeclaredName);
+//		int64 Idx = DeclaredName - &Declaration.Names[0];
+//
+//		check(TypeIndex < UE_ARRAY_COUNT(Loaders));
+//		Loaders[TypeIndex](Ar, OutVariant);		
+//		
+//	}
+//
+//	template<typename T>
+//	void Load(TVariant<Ts...>& Dst, FCustomMemberLoader& Src, ECustomLoadMethod Method)
+//	{
+//
+//		
+//		{
+//			Dst.Emplace(MoveTemp(*reinterpret_cast<T*>(Temp)));
+//		}
+//		else
+//		{
+//			new (Dst) TVariant<Ts...>(MoveTemp(*reinterpret_cast<T*>(Temp)));
+//		}
+//	}
+//
+//	virtual TUniquePtr<FBuiltStruct>	SaveStruct(const void* Src, const FDebugIds& Debug) const override
+//	{
+//		...
+//	}
+//};
+
+}
+
+namespace PlainProps
+{
+	template<typename T>
+	struct TRangeBind<TArray<T>>
+	{
+		using Type = UE::TArrayBinding<T>;
+	};
+
+	template<typename T>
+	struct TRangeBind<TSet<T>>
+	{
+		using Type = UE::TSetBinding<T>;
+	};
+}
\ No newline at end of file
diff --git a/Engine/Source/Programs/PlainPropsTests/Public/PlainPropsWrite.h b/Engine/Source/Programs/PlainPropsTests/Public/PlainPropsWrite.h
new file mode 100644
--- /dev/null
+++ b/Engine/Source/Programs/PlainPropsTests/Public/PlainPropsWrite.h
@@ -0,0 +1,59 @@
+// Copyright Epic Games, Inc. All Rights Reserved.
+
+#pragma once
+
+#include "Containers/Array.h"
+#include  "PlainPropsTypes.h"
+#include "Templates/UniquePtr.h"
+
+namespace PlainProps
+{
+struct FBuiltSchemas;
+struct FBuiltStruct;
+class FDebugIds;
+class FIdIndexerBase;
+class FNestedScopeIndexer;
+class FParametricTypeIndexer;
+struct FWriteIds;
+
+///// Maps declared / built ids to write ids, the subset that's used by the partial schemas
+//class FUsedIds
+//{
+//public:
+//	virtual ~FUsedIds() {}
+//	virtual bool Uses(FNameId BuiltId) const = 0;
+//	virtual FOptionalStructSchemaId GetWriteId(FStructSchemaId BuiltId) const = 0;
+//};
+//
+//PP_API TUniquePtr<FUsedIds>	DropUnusedIds(const FBuiltSchemas& AllSchemas, int32 NumDeclaredNames);
+//PP_API void					WriteSchemas(TArray64<uint8>& Out, const FBuiltSchemas& AllSchemas, const FUsedIds& WriteIds);
+//PP_API void					WriteMembers(TArray64<uint8>& Out, const FBuiltSchemas& AllSchemas, const FUsedIds& WriteIds, FStructSchemaId Schema, const FBuiltStruct& Struct, const FDebugIds& Debug);
+
+//struct FDeclaredIds
+//{
+//	int32 NumNames;
+//	const FNestedScopeIndexer& NestedScopes;
+//	const FParametricTypeIndexer& ParametricTypes;
+//};
+
+class FWriter
+{
+public:
+	PP_API FWriter(const FIdIndexerBase& DeclaredIds, const FBuiltSchemas& InSchemas);
+	PP_API ~FWriter();
+	
+	PP_API bool								Uses(FNameId BuiltId) const;
+	PP_API TConstArrayView<FNestedScope>		GetUsedScopes() const;
+	PP_API TConstArrayView<FParametricType>	GetUsedParametrics() const;
+	PP_API FOptionalStructSchemaId			GetWriteId(FStructSchemaId BuiltId) const;
+
+	PP_API void								WriteSchemas(TArray64<uint8>& Out) const;
+	PP_API void								WriteMembers(TArray64<uint8>& Out, FStructSchemaId BuiltId, const FBuiltStruct& Struct) const;
+
+private:
+	const FBuiltSchemas& Schemas;
+	const FDebugIds& Debug;
+	TUniquePtr<FWriteIds> NewIds;
+};
+
+} // namespace PlainProps
\ No newline at end of file
diff --git a/Engine/Source/Programs/PlainPropsTests/Tests/PlainPropsCttiTests.cpp b/Engine/Source/Programs/PlainPropsTests/Tests/PlainPropsCttiTests.cpp
new file mode 100644
--- /dev/null
+++ b/Engine/Source/Programs/PlainPropsTests/Tests/PlainPropsCttiTests.cpp
@@ -0,0 +1,126 @@
+// Copyright Epic Games, Inc. All Rights Reserved.
+
+#include  "PlainPropsCtti.h"
+#include <type_traits>
+#include <string_view>
+
+namespace PlainProps::Test
+{
+
+static_assert(std::is_same_v<CttiOf<float>::Type, float>);
+static_assert(std::string_view(CttiOf<int>::Name) == std::string_view("int32_t"));
+
+//////////////////////////////////////////////////////////////////////////
+
+enum class E1 : uint8_t { A, B };
+
+namespace Actual
+{	
+	PP_REFLECT_ENUM(PlainProps::Test, E1, A, B);
+}
+
+namespace Expect
+{
+	struct E1_Ctti
+	{
+		static constexpr char Name[] = "E1";
+		using Type = ::PlainProps::Test::E1;
+		static constexpr int NumEnumerators = 2;
+		static constexpr struct { const char* Name; Type Value; } Enumerators[] = { {"A", Type::A},  {"B", Type::B} };
+	};
+}
+
+static_assert(std::string_view(Actual::E1_Ctti::Name) == std::string_view(Expect::E1_Ctti::Name));
+static_assert(std::is_same_v<Actual::E1_Ctti::Type, Expect::E1_Ctti::Type>);
+static_assert(Actual::E1_Ctti::NumEnumerators == Expect::E1_Ctti::NumEnumerators);
+static_assert(Actual::E1_Ctti::Enumerators[0].Value	== Expect::E1_Ctti::Enumerators[0].Value);
+static_assert(Actual::E1_Ctti::Enumerators[1].Value	== Expect::E1_Ctti::Enumerators[1].Value);
+static_assert(std::string_view(Actual::E1_Ctti::Enumerators[0].Name) == std::string_view(Expect::E1_Ctti::Enumerators[0].Name));
+static_assert(std::string_view(Actual::E1_Ctti::Enumerators[1].Name) == std::string_view(Expect::E1_Ctti::Enumerators[1].Name));
+
+
+//////////////////////////////////////////////////////////////////////////
+
+struct S1
+{
+	float x;
+	int y;
+};
+
+namespace Actual
+{	
+	PP_REFLECT_STRUCT(PlainProps::Test, S1, void, x, y);
+}
+
+namespace Expect
+{
+	struct S1_Ctti
+	{
+		static constexpr char Name[] = "S1";
+		using Type = ::PlainProps::Test::S1;
+		using Super = void;
+		static constexpr int NumVars = 2;
+		template<int> struct Var;
+	};
+
+	template<> struct S1_Ctti::Var<2-2>
+	{
+		static constexpr char Name[] = "x";
+		using Type = decltype(::PlainProps::Test::S1::x);
+		static constexpr auto Pointer = &::PlainProps::Test::S1::x;
+		static constexpr std::size_t Offset = offsetof(::PlainProps::Test::S1, x);
+		static constexpr int Index = 2-2;
+	};
+
+	template<> struct S1_Ctti::Var<2-1>
+	{
+		static constexpr char Name[] = "y";
+		using Type = decltype(::PlainProps::Test::S1::y);
+		static constexpr auto Pointer = &::PlainProps::Test::S1::y;
+		static constexpr std::size_t Offset = offsetof(::PlainProps::Test::S1, y);
+		static constexpr int Index = 2-1;
+	};
+}
+
+template<class Actual, class Expect>
+static constexpr bool AssertVarEquivalence()
+{
+	static_assert(std::string_view(Actual::Name) == std::string_view(Expect::Name));
+	static_assert(std::is_same_v<typename Actual::Type, typename Expect::Type>);
+	static_assert(Actual::Offset == Expect::Offset);
+	static_assert(Actual::Pointer == Expect::Pointer);
+	static_assert(Actual::Index == Expect::Index);
+	return true;
+}
+
+static_assert(std::string_view(Actual::S1_Ctti::Name) == std::string_view(Expect::S1_Ctti::Name));
+static_assert(std::is_same_v<Actual::S1_Ctti::Type, Expect::S1_Ctti::Type>);
+static_assert(std::is_same_v<Actual::S1_Ctti::Super, Expect::S1_Ctti::Super>);
+static_assert(Actual::S1_Ctti::NumVars == Expect::S1_Ctti::NumVars);
+static_assert(AssertVarEquivalence<Actual::S1_Ctti::Var<0>, Expect::S1_Ctti::Var<0>>());
+static_assert(AssertVarEquivalence<Actual::S1_Ctti::Var<1>, Expect::S1_Ctti::Var<1>>());
+
+// CttiOf only works when CTTI exists in same or parent namespace, it uses argument dependent lookup
+// to find the "canonical" CTTI if exists. Regenerate S1_Ctti in S1's namespace to test it.
+PP_REFLECT_STRUCT(PlainProps::Test, S1, void, x, y);
+static_assert(std::is_same_v<CttiOf<S1>, S1_Ctti>);
+
+//////////////////////////////////////////////////////////////////////////
+
+template<class T>
+struct S2
+{
+	bool _; // unreflected
+	T a;
+};
+
+PP_REFLECT_STRUCT_TEMPLATE(PlainProps::Test, S2, void, a);
+
+static_assert(std::string_view(CttiOf<S2<int>>::Name) == std::string_view("S2"));
+static_assert(std::is_same_v<CttiOf<S2<int>>::Type, S2<int>>);
+static_assert(std::is_same_v<CttiOf<S2<int>>::TemplateArgs, std::tuple<int>>);
+static_assert(CttiOf<S2<int>>::NumVars == 1);
+static_assert(CttiOf<S2<int>>::Var<0>::Name == std::string_view("a"));
+static_assert(CttiOf<S2<int>>::Var<0>::Offset == offsetof(S2<int>, a));
+
+}
\ No newline at end of file
diff --git a/Engine/Source/Programs/PlainPropsTests/Tests/PlainPropsTests.cpp b/Engine/Source/Programs/PlainPropsTests/Tests/PlainPropsTests.cpp
new file mode 100644
--- /dev/null
+++ b/Engine/Source/Programs/PlainPropsTests/Tests/PlainPropsTests.cpp
@@ -0,0 +1,1432 @@
+// Copyright Epic Games, Inc. All Rights Reserved.
+
+//#if WITH_TESTS //TMPTORP
+#include "Tests/TestHarnessAdapter.h"
+#include "Containers/AnsiString.h"
+#include "Containers/Map.h"
+#include "Containers/Set.h"
+#include "Containers/StringView.h"
+#include  "PlainPropsRead.h"
+#include  "PlainPropsWrite.h"
+#include  "PlainPropsBuildSchema.h"
+#include  "PlainPropsIndex.h"
+#include  "PlainPropsInternalBuild.h"
+#include  "PlainPropsInternalFormat.h"
+#include "Templates/UnrealTemplate.h"
+
+namespace PlainProps
+{
+	
+static bool operator==(FScopeId A, FNestedScopeId B) { return A == FScopeId(B); }
+static bool operator==(FScopeId A, FFlatScopeId B) { return A == FScopeId(B); }
+static bool operator==(FParametricTypeView A, FParametricTypeView B)
+{
+	return A.Name == B.Name && A.NumParameters == B.NumParameters && !FMemory::Memcmp(A.Parameters, B.Parameters, A.NumParameters * sizeof(FTypeId));
+}
+	
+TEST_CASE_NAMED(FPlainPropsIndexTest, "System::Core::Serialization::PlainProps::Index", "[Core][PlainProps][SmokeFilter]")
+{
+	SECTION("NestedScope")
+	{
+		FScopeId S0{FFlatScopeId{{0}}};
+		FScopeId S1{FFlatScopeId{{1}}};
+		FScopeId S2{FFlatScopeId{{2}}};
+
+		FNestedScope N01{S0, S1.AsFlat()};
+		FNestedScope N10{S1, S0.AsFlat()};
+		FNestedScope N12{S1, S2.AsFlat()};
+
+		FNestedScopeIndexer Indexer;
+
+		FScopeId S01(Indexer.Index(N01));
+		FScopeId S10(Indexer.Index(N10));
+		FScopeId S12(Indexer.Index(N12));
+
+		FNestedScope N012{S01, S2.AsFlat()};
+		FScopeId S012(Indexer.Index(N012));
+
+		FNestedScope N0120{S012, S0.AsFlat()};
+		FScopeId S0120(Indexer.Index(N0120));
+		
+		CHECK(S01	== Indexer.Index(N01));
+		CHECK(S10	== Indexer.Index(N10));
+		CHECK(S12	== Indexer.Index(N12));
+		CHECK(S012	== Indexer.Index(N012));
+		CHECK(S0120 == Indexer.Index(N0120));
+		CHECK(N01	== Indexer.Resolve(S01.AsNested()));
+		CHECK(N10	== Indexer.Resolve(S10.AsNested()));
+		CHECK(N12	== Indexer.Resolve(S12.AsNested()));
+		CHECK(N012	== Indexer.Resolve(S012.AsNested()));
+		CHECK(N0120 == Indexer.Resolve(S0120.AsNested()));
+		CHECK(Indexer.Num() == 5);
+	}
+	
+	SECTION("ParametricType")
+	{
+		FScopeId S0{FFlatScopeId{{0}}};
+		FScopeId S1{FFlatScopeId{{1}}};
+		FScopeId S2{FFlatScopeId{{2}}};
+		
+		FConcreteTypenameId T3{{3}};
+		FConcreteTypenameId T4{{4}};
+		FConcreteTypenameId T5{{5}};
+		
+		FTypeId S0T3 = {S0, FTypenameId{T3}};
+		FTypeId S1T3 = {S1, FTypenameId{T3}};
+		FTypeId S2T3 = {S2, FTypenameId{T3}};
+
+		FParametricTypeIndexer Indexer;
+
+		FParametricTypeId T4_S0T3 = Indexer.Index({T4, 1, &S0T3});
+		FParametricTypeId T4_S1T3 = Indexer.Index({T4, 1, &S1T3});
+		
+		CHECK(Indexer.Resolve(T4_S0T3) == FParametricTypeView{T4, 1, &S0T3});
+		CHECK(Indexer.Resolve(T4_S1T3) == FParametricTypeView{T4, 1, &S1T3});
+		
+		FTypeId S1T4_S0T3 = {S1,  FTypenameId{T4_S0T3}};
+		FTypeId S2T4_S1T3 = {S2,  FTypenameId{T4_S1T3}};
+		
+		CHECK(S1T4_S0T3.Name.AsParametric() == T4_S0T3);
+		CHECK(S2T4_S1T3.Name.AsParametric() == T4_S1T3);
+
+		FParametricTypeId T5_S0T3_S2T3 = Indexer.Index({T5, {S1T4_S0T3, S2T4_S1T3}});
+		CHECK(Indexer.Resolve(T5_S0T3_S2T3) == FParametricTypeView{T5, {S1T4_S0T3, S2T4_S1T3}});
+		
+		CHECK(T4_S0T3		== Indexer.Index({T4, 1, &S0T3}));
+		CHECK(T4_S1T3		== Indexer.Index({T4, 1, &S1T3}));
+		CHECK(T5_S0T3_S2T3	== Indexer.Index({T5, {S1T4_S0T3, S2T4_S1T3}}));
+
+		CHECK(Indexer.Num() == 3);
+	}
+}
+
+////////////////////////////////////////////////////////////////////////////////////////////////
+
+inline constexpr uint32 TestMagics[] = {0xFEEDF00D, 0xABCD1234, 0xDADADAAA, 0x99887766, 0xF0F1F2F3 };
+
+class FTestBatchBuilder : public TIdIndexer<FAnsiString>, public FDeclarations
+{
+public:
+	FEnumSchemaId				DeclareEnum(FTypeId Type, EEnumMode Mode, ELeafWidth Width, std::initializer_list<const char*> Names, std::initializer_list<uint64> Constants);
+	FEnumSchemaId				DeclareEnum(const char* Scope, const char* Name, EEnumMode Mode, ELeafWidth Width, std::initializer_list<const char*> Names, std::initializer_list<uint64> Constants);
+	FStructSchemaId				DeclareStruct(FTypeId Type, std::initializer_list<const char*> MemberOrder, EMemberPresence Occupancy, FOptionalStructSchemaId Super = {});
+	FStructSchemaId				DeclareStruct(const char* Scope, const char* Name, std::initializer_list<const char*> MemberOrder, EMemberPresence Occupancy, FOptionalStructSchemaId Super = {});
+
+	void						AddObject(FStructSchemaId Schema, FMemberBuilder&& Members);
+	void						AddObjects(FStructSchemaId Schema, std::initializer_list<FMemberBuilder&&> Objects);
+
+	TArray64<uint8>				Write();
+
+private:
+	TArray<TPair<FStructSchemaId, TUniquePtr<FBuiltStruct>>>	Objects;
+
+	TArray<FMemberId> NameMembers(std::initializer_list<const char*> Members)
+	{
+		TArray<FMemberId> Out;
+		Out.Reserve(Members.size());
+		for (const char* Member : Members)
+		{
+			Out.Add(NameMember(Member));
+		}
+		return Out;
+	}
+	
+	TArray<FEnumerator> MakeEnumerators(TConstArrayView<const char*> InNames, TConstArrayView<uint64> Constants)
+	{
+		check(InNames.Num() == Constants.Num());
+		TArray<FEnumerator> Out;
+		Out.SetNumUninitialized(InNames.Num());
+		for (uint32 Idx = 0, Num = Out.Num(); Idx < Num; ++Idx)
+		{
+			Out[Idx] = { MakeName(InNames[Idx]), Constants[Idx] };
+		}
+		return Out;
+	}
+
+	TArray<char> GetNameData() const;
+};
+
+FStructSchemaId FTestBatchBuilder::DeclareStruct(FTypeId Type, std::initializer_list<const char*> MemberOrder, EMemberPresence Occupancy, FOptionalStructSchemaId Super)
+{
+	FStructSchemaId Id = IndexStruct(Type);
+	FDeclarations::DeclareStruct(Id, Type, NameMembers(MemberOrder), Occupancy, Super);
+	return Id;
+}
+
+FStructSchemaId FTestBatchBuilder::DeclareStruct(const char* Scope, const char* Name, std::initializer_list<const char*> MemberOrder, EMemberPresence Occupancy, FOptionalStructSchemaId Super)
+{
+	return DeclareStruct(MakeType(Scope, Name), MemberOrder, Occupancy, Super);
+}
+
+FEnumSchemaId FTestBatchBuilder::DeclareEnum(FTypeId Type, EEnumMode Mode, ELeafWidth Width, std::initializer_list<const char*> InNames, std::initializer_list<uint64> Constants)
+{
+	FEnumSchemaId Id = IndexEnum(Type);
+	FDeclarations::DeclareEnum(Id, Type, Mode, Width, MakeEnumerators(InNames, Constants));
+	return Id;
+}
+
+FEnumSchemaId FTestBatchBuilder::DeclareEnum(const char* Scope, const char* Name, EEnumMode Mode, ELeafWidth Width, std::initializer_list<const char*> InNames, std::initializer_list<uint64> Constants)
+{
+	return DeclareEnum(MakeType(Scope, Name), Mode, Width, InNames, Constants);
+}
+
+void FTestBatchBuilder::AddObject(FStructSchemaId Schema, FMemberBuilder&& Members)
+{
+	const FStructDeclaration& Declaration = *DeclaredStructs[Schema.Idx];
+	Objects.Emplace(Schema, Members.BuildAndReset(Declaration, *this));
+}
+
+TArray64<uint8> FTestBatchBuilder::Write()
+{
+	// Build partial schemas
+	FSchemasBuilder SchemaBuilders(DeclaredStructs, DeclaredEnums, *this);
+	for (const TPair<FStructSchemaId, TUniquePtr<FBuiltStruct>>& Object : Objects)
+	{
+		SchemaBuilders.NoteMembers(Object.Key, *Object.Value);
+	}
+	FBuiltSchemas Schemas = SchemaBuilders.Build(); 
+
+	// Filter out declared but unused names and ids
+	FWriter Writer(*this, Schemas);
+
+	// Write names
+	TArray64<uint8> Out;
+	TArray64<uint8> Tmp;
+	uint32 OldNameIdx = 0;
+	for (const FAnsiString& Name : Names)
+	{
+		if (Writer.Uses({OldNameIdx++}))
+		{
+			WriteData(Tmp, *Name, Name.Len() + 1);
+		}
+	}
+	WriteU32(Out, TestMagics[0]);
+	WriteSkippableSlice(Out, Tmp);
+	Tmp.Reset();
+
+	// Write schemas
+	WriteU32(Out, TestMagics[1]);
+	Writer.WriteSchemas(/* Out */ Tmp);
+	WriteAlignmentPadding<uint32>(Out);
+	WriteU32(Out, IntCastChecked<uint32>(Tmp.Num()));
+	WriteArray(Out, Tmp);
+	Tmp.Reset();
+
+	// Write objects
+	WriteU32(Out, TestMagics[2]);
+	for (const TPair<FStructSchemaId, TUniquePtr<FBuiltStruct>>& Object : Objects)
+	{
+		WriteU32(/* out */ Tmp, TestMagics[3]);
+		WriteU32(/* out */ Tmp, Writer.GetWriteId(Object.Key).Get().Idx);
+		Writer.WriteMembers(/* out */ Tmp, Object.Key, *Object.Value);
+		WriteSkippableSlice(Out, Tmp);
+		Tmp.Reset();
+	}
+
+	// Write object terminator
+	WriteSkippableSlice(Out, TConstArrayView64<uint8>());
+	WriteU32(Out, TestMagics[4]);
+		
+	return Out;
+}
+
+TArray<char> FTestBatchBuilder::GetNameData() const
+{
+	TArray<char> Out;
+	Out.Reserve(Names.Num() * 100);
+	for (const FAnsiString& Name : Names)
+	{
+		Out.Append(Name.GetCharArray());
+		check(Out.Last() == '\0');
+	}
+	return Out;
+}
+
+////////////////////////////////////////////////////////////////////////////////////////////////
+
+class FTestNameReader
+{
+public:
+	void Read(FMemoryView Data)
+	{
+		check(Names.IsEmpty() && !Data.IsEmpty());
+		TConstArrayView<char> AllChars(static_cast<const char*>(Data.GetData()), IntCastChecked<int32>(Data.GetSize()));
+		
+		const char* NextStr = AllChars.GetData();
+		for (const char& Char : AllChars)
+		{
+			if (!Char)
+			{
+				Names.Add(NextStr);
+				NextStr = &Char + 1;
+			}
+		}
+		
+		check(Names.Num() >= 3); // 1 FTypeId and 1 member id at least
+		check(NextStr == Data.GetDataEnd()); // end with null-terminator
+	}
+	
+	FAnsiStringView operator[](FNameId Id) const { return MakeStringView(Names[Id.Idx]); }
+	FAnsiStringView operator[](FMemberId Name) const { return MakeStringView(Names[Name.Id.Idx]); }
+	FAnsiStringView operator[](FOptionalMemberId Name) const { return Name ? operator[](Name.Get()) : "Super"; }
+	FAnsiStringView operator[](FScopeId Scope) const { return operator[](Scope.AsFlat().Name); }
+	FAnsiStringView operator[](FTypenameId Name) const { return operator[](Name.AsConcrete().Id); }
+
+private:
+	TArray<const char*> Names;
+};
+
+class FTestBatchReader
+{
+public:
+	FTestBatchReader(FMemoryView Data)
+	{
+		// Read names
+		FByteReader It(Data);
+		CHECK(It.Grab<uint32>() == TestMagics[0]);
+		Names.Read(It.GrabSkippableSlice());
+		
+		// Read schemas
+		CHECK(It.Grab<uint32>() == TestMagics[1]);
+		It.SkipAlignmentPadding<uint32>();
+		uint32 SchemasSize = It.Grab<uint32>();
+		const FSchemaBatch* Schemas = ValidateSchemas(It.GrabSlice(SchemasSize));
+		FReadBatchId Batch = MountReadSchemas(Schemas);
+		CHECK(It.Grab<uint32>() == TestMagics[2]);
+		
+		// Read objects
+		while (uint64 NumBytes = It.GrabVarIntU())
+		{	
+			FByteReader ObjIt(It.GrabSlice(NumBytes));
+			CHECK(ObjIt.Grab<uint32>() == TestMagics[3]);
+			FStructSchemaId Schema = { ObjIt.Grab<uint32>() };
+			Objects.Add({ { Schema, Batch }, ObjIt });
+		}
+		
+		CHECK(It.Grab<uint32>() == TestMagics[4]);
+		CHECK(!Objects.IsEmpty());
+	}
+
+	~FTestBatchReader()
+	{
+		UnmountReadSchemas(Objects[0].Schema.Batch);
+	}
+
+	TConstArrayView<FStructView>		GetObjects() const { return Objects;	}
+	const FTestNameReader&				GetNames() const { return Names; }
+	
+private:
+	FTestNameReader Names;
+	TArray<FStructView> Objects;
+};
+
+static void TestSerialize(void (*BuildObjects)(FTestBatchBuilder&), void (*CheckObjects)(TConstArrayView<FStructView>, const FTestNameReader&))
+{
+	TArray64<uint8> Data;
+	{
+		FTestBatchBuilder Batch;
+		BuildObjects(Batch);
+		Data = Batch.Write();
+	}
+
+	FTestBatchReader Batch(MakeMemoryView(Data));
+	CheckObjects(Batch.GetObjects(), Batch.GetNames());
+}
+
+////////////////////////////////////////////////////////////////////////////////////////////////
+
+// Tests that everything was read
+struct FTestMemberReader : public FMemberReader
+{
+	using FMemberReader::FMemberReader;
+	
+	~FTestMemberReader()
+	{
+		CHECK(MemberIdx == NumMembers); // Must read all members
+		CHECK(RangeTypeIdx == NumRangeTypes); // Must read all ranges
+#if DO_CHECK
+		CHECK(InnerSchemaIdx == NumInnerSchemas); // Must read all schema ids
+#endif
+	}
+};
+
+template<typename OutType, typename InType>
+TArray<OutType> MakeArray(const InType& Items)
+{
+	TArray<OutType> Out;
+	Out.Reserve(IntCastChecked<int32>(Items.Num()));
+	for (const auto& Item : Items)
+	{
+		Out.Emplace(Item);
+	}
+	return Out;
+}
+
+////////////////////////////////////////////////////////////////////////////////////////////////
+
+TEST_CASE_NAMED(FPlainPropsReadWriteTest, "System::Core::Serialization::PlainProps::ReadWrite", "[Core][PlainProps][SmokeFilter]")
+{
+	SECTION("Bool")
+	{
+		static constexpr std::initializer_list<const char*> MemberNames =  {"b0", "b1", "b2", "b3", "b4", "b5", "b6", "b7", "b8", "b9", "b10", "b11"};
+
+		TestSerialize([](FTestBatchBuilder& Batch)
+		{
+			FStructSchemaId SchemaId = Batch.DeclareStruct("Testing", "Bools", MemberNames, EMemberPresence::AllowSparse);
+
+			FMemberBuilder B1T;
+			B1T.Add(Batch.NameMember("b3"),	true);
+
+			FMemberBuilder B1F;
+			B1F.Add(Batch.NameMember("b1"),	false);
+
+			FMemberBuilder B8M;
+			B8M.Add(Batch.NameMember("b1"),	true);
+			B8M.Add(Batch.NameMember("b2"),	false);
+			B8M.Add(Batch.NameMember("b3"),	true);
+			B8M.Add(Batch.NameMember("b4"),	false);
+			B8M.Add(Batch.NameMember("b5"),	true);
+			B8M.Add(Batch.NameMember("b6"),	false);
+			B8M.Add(Batch.NameMember("b8"),	false);
+			B8M.Add(Batch.NameMember("b9"), true);
+			
+			FMemberBuilder B9T;
+			B9T.Add(Batch.NameMember("b1"),	true);
+			B9T.Add(Batch.NameMember("b2"),	true);
+			B9T.Add(Batch.NameMember("b3"),	true);
+			B9T.Add(Batch.NameMember("b4"),	true);
+			B9T.Add(Batch.NameMember("b5"),	true);
+			B9T.Add(Batch.NameMember("b6"),	true);
+			B9T.Add(Batch.NameMember("b8"),	true);
+			B9T.Add(Batch.NameMember("b9"), true);
+			B9T.Add(Batch.NameMember("b10"),true);
+
+			Batch.AddObject(SchemaId, MoveTemp(B1T));
+			Batch.AddObject(SchemaId, MoveTemp(B1F));
+			Batch.AddObject(SchemaId, MoveTemp(B8M));
+			Batch.AddObject(SchemaId, MoveTemp(B9T));
+		}, 
+		[](TConstArrayView<FStructView> Objects, const FTestNameReader& Names)
+		{
+			CHECK(Objects.Num() == 4);
+			FTestMemberReader B1T(Objects[0]);
+			FTestMemberReader B1F(Objects[1]);
+			FTestMemberReader B8M(Objects[2]);
+			FTestMemberReader B9T(Objects[3]);
+			CHECK(Objects[0].Schema.Id == Objects[3].Schema.Id);
+
+			// Check schema
+			const FStructSchema& Schema =  Objects[0].Schema.Resolve();
+			CHECK(Names[Schema.Type.Scope] == "Testing");
+			CHECK(Names[Schema.Type.Name] == "Bools");
+			CHECK(Schema.NumMembers == 9); // b0, b7 and b11 unused
+			CHECK(Schema.NumRangeTypes == 0);
+			CHECK(Schema.NumInnerSchemas == 0);
+			CHECK(Schema.IsDense == 0);
+			CHECK(Schema.Inheritance == ESuper::No);
+			CHECK(FStructSchema::GetMemberTypes(Schema.Footer)[0] == FUnpackedLeafType(ELeafType::Bool, ELeafWidth::B8).Pack());
+			CHECK(FStructSchema::GetMemberTypes(Schema.Footer)[8] == FUnpackedLeafType(ELeafType::Bool, ELeafWidth::B8).Pack());
+			TConstArrayView<FMemberId> MemberIds(FStructSchema::GetMemberNames(Schema.Footer, Schema.NumMembers, Schema.NumRangeTypes), Schema.NumMembers);
+			CHECK(Names[MemberIds[0]] == "b1");
+			CHECK(Names[MemberIds[1]] == "b2");
+			CHECK(Names[MemberIds[2]] == "b3");
+			CHECK(Names[MemberIds[3]] == "b4");
+			CHECK(Names[MemberIds[4]] == "b5");
+			CHECK(Names[MemberIds[5]] == "b6");
+			CHECK(Names[MemberIds[6]] == "b8");
+			CHECK(Names[MemberIds[7]] == "b9");
+			CHECK(Names[MemberIds[8]] == "b10");
+
+			CHECK(Names[B1T.PeekName()] == "b3");
+			CHECK(B1T.GrabLeaf().AsBool() == true);
+
+			CHECK(Names[B1F.PeekName()] == "b1");
+			CHECK(B1F.GrabLeaf().AsBool() == false);
+			
+			CHECK(B8M.GrabLeaf().AsBool() == true);
+			CHECK(B8M.GrabLeaf().AsBool() == false);
+			CHECK(B8M.GrabLeaf().AsBool() == true);
+			CHECK(B8M.GrabLeaf().AsBool() == false);
+			CHECK(B8M.GrabLeaf().AsBool() == true);
+			CHECK(B8M.GrabLeaf().AsBool() == false);
+			CHECK(B8M.GrabLeaf().AsBool() == false);
+			CHECK(B8M.GrabLeaf().AsBool() == true);
+
+			for (int32 Idx = 0; Idx < 9; ++Idx)
+			{
+				CHECK(B9T.GrabLeaf().AsBool() == true);
+			}
+		});
+	}
+
+	SECTION("Number")
+	{
+		static constexpr std::initializer_list<const char*> MemberNames =  {"F32", "F64", "S8", "U8", "S16", "U16", "S32", "U32", "S64", "U64"};
+
+		TestSerialize([](FTestBatchBuilder& Batch)
+		{
+			FStructSchemaId SchemaId = Batch.DeclareStruct("Test", "Numbers", MemberNames, EMemberPresence::AllowSparse);
+
+			FMemberBuilder Misc, Mins, Maxs, Some;
+
+			Misc.Add(Batch.NameMember("F32"),	32.f);
+			Misc.Add(Batch.NameMember("F64"),	64.0);
+			Misc.Add(Batch.NameMember("S8"),	int8(-8));
+			Misc.Add(Batch.NameMember("U8"),	uint8(8));
+			Misc.Add(Batch.NameMember("S16"),	int16(-16));
+			Misc.Add(Batch.NameMember("U16"),	uint16(16));
+			Misc.Add(Batch.NameMember("S32"),	int32(-32));
+			Misc.Add(Batch.NameMember("U32"),	uint32(32));
+			Misc.Add(Batch.NameMember("S64"),	int64(-64));
+			Misc.Add(Batch.NameMember("U64"),	uint64(64));
+			
+			Mins.Add(Batch.NameMember("F32"),	std::numeric_limits< float>::min());
+			Mins.Add(Batch.NameMember("F64"),	std::numeric_limits<double>::min());
+			Mins.Add(Batch.NameMember("S8"),	std::numeric_limits<  int8>::min());
+			Mins.Add(Batch.NameMember("U8"),	std::numeric_limits< uint8>::min());
+			Mins.Add(Batch.NameMember("S16"),	std::numeric_limits< int16>::min());
+			Mins.Add(Batch.NameMember("U16"),	std::numeric_limits<uint16>::min());
+			Mins.Add(Batch.NameMember("S32"),	std::numeric_limits< int32>::min());
+			Mins.Add(Batch.NameMember("U32"),	std::numeric_limits<uint32>::min());
+			Mins.Add(Batch.NameMember("S64"),	std::numeric_limits< int64>::min());
+			Mins.Add(Batch.NameMember("U64"),	std::numeric_limits<uint64>::min());
+
+			Maxs.Add(Batch.NameMember("F32"),	std::numeric_limits< float>::max());
+			Maxs.Add(Batch.NameMember("F64"),	std::numeric_limits<double>::max());
+			Maxs.Add(Batch.NameMember("S8"),	std::numeric_limits<  int8>::max());
+			Maxs.Add(Batch.NameMember("U8"),	std::numeric_limits< uint8>::max());
+			Maxs.Add(Batch.NameMember("S16"),	std::numeric_limits< int16>::max());
+			Maxs.Add(Batch.NameMember("U16"),	std::numeric_limits<uint16>::max());
+			Maxs.Add(Batch.NameMember("S32"),	std::numeric_limits< int32>::max());
+			Maxs.Add(Batch.NameMember("U32"),	std::numeric_limits<uint32>::max());
+			Maxs.Add(Batch.NameMember("S64"),	std::numeric_limits< int64>::max());
+			Maxs.Add(Batch.NameMember("U64"),	std::numeric_limits<uint64>::max());
+
+			Some.Add(Batch.NameMember("S32"),	0);
+			
+			Batch.AddObject(SchemaId, MoveTemp(Misc));
+			Batch.AddObject(SchemaId, MoveTemp(Mins));
+			Batch.AddObject(SchemaId, MoveTemp(Maxs));
+			Batch.AddObject(SchemaId, MoveTemp(Some));
+		}, 
+		[](TConstArrayView<FStructView> Objects, const FTestNameReader& Names)
+		{
+			for (const FStructView& Object : Objects.Left(3))
+			{
+				FTestMemberReader Members(Object);
+				for (const char* MemberName : MemberNames)
+				{
+					CHECK(Members.HasMore());
+					CHECK(Names[Members.PeekName()] == MemberName);
+					CHECK(Members.PeekKind() == EMemberKind::Leaf);
+					(void)Members.GrabLeaf();
+				}
+			}
+			
+			FTestMemberReader Misc(Objects[0]);
+			CHECK(Misc.GrabLeaf().AsFloat()		== 32.f);
+			CHECK(Misc.GrabLeaf().AsDouble()	== 64.0);
+			CHECK(Misc.GrabLeaf().AsS8()		== int8(-8));
+			CHECK(Misc.GrabLeaf().AsU8()		== uint8(8));
+			CHECK(Misc.GrabLeaf().AsS16()		== int16(-16));
+			CHECK(Misc.GrabLeaf().AsU16()		== uint16(16));
+			CHECK(Misc.GrabLeaf().AsS32()		== int32(-32));
+			CHECK(Misc.GrabLeaf().AsU32()		== uint32(32));
+			CHECK(Misc.GrabLeaf().AsS64()		== int64(-64));
+			CHECK(Misc.GrabLeaf().AsU64()		== uint64(64));
+
+			FTestMemberReader Mins(Objects[1]);
+			CHECK(Mins.GrabLeaf().AsFloat()		== std::numeric_limits< float>::min());
+			CHECK(Mins.GrabLeaf().AsDouble()	== std::numeric_limits<double>::min());
+			CHECK(Mins.GrabLeaf().AsS8()		== std::numeric_limits<  int8>::min());
+			CHECK(Mins.GrabLeaf().AsU8()		== std::numeric_limits< uint8>::min());
+			CHECK(Mins.GrabLeaf().AsS16()		== std::numeric_limits< int16>::min());
+			CHECK(Mins.GrabLeaf().AsU16()		== std::numeric_limits<uint16>::min());
+			CHECK(Mins.GrabLeaf().AsS32()		== std::numeric_limits< int32>::min());
+			CHECK(Mins.GrabLeaf().AsU32()		== std::numeric_limits<uint32>::min());
+			CHECK(Mins.GrabLeaf().AsS64()		== std::numeric_limits< int64>::min());
+			CHECK(Mins.GrabLeaf().AsU64()		== std::numeric_limits<uint64>::min());
+
+			FTestMemberReader Maxs(Objects[2]);
+			CHECK(Maxs.GrabLeaf().AsFloat()		== std::numeric_limits< float>::max());
+			CHECK(Maxs.GrabLeaf().AsDouble()	== std::numeric_limits<double>::max());
+			CHECK(Maxs.GrabLeaf().AsS8()		== std::numeric_limits<  int8>::max());
+			CHECK(Maxs.GrabLeaf().AsU8()		== std::numeric_limits< uint8>::max());
+			CHECK(Maxs.GrabLeaf().AsS16()		== std::numeric_limits< int16>::max());
+			CHECK(Maxs.GrabLeaf().AsU16()		== std::numeric_limits<uint16>::max());
+			CHECK(Maxs.GrabLeaf().AsS32()		== std::numeric_limits< int32>::max());
+			CHECK(Maxs.GrabLeaf().AsU32()		== std::numeric_limits<uint32>::max());
+			CHECK(Maxs.GrabLeaf().AsS64()		== std::numeric_limits< int64>::max());
+			CHECK(Maxs.GrabLeaf().AsU64()		== std::numeric_limits<uint64>::max());
+
+			FTestMemberReader Some(Objects[3]);
+			CHECK(Names[Some.PeekName()] == "S32");
+			CHECK(Some.GrabLeaf().AsS32() == 0);
+		});
+	}
+
+	SECTION("Dense")
+	{
+		static constexpr std::initializer_list<const char*> MemberNames =  {"A", "B", "C"};
+
+		TestSerialize([](FTestBatchBuilder& Batch)
+		{
+			FStructSchemaId ExplicitId = Batch.DeclareStruct("Test", "ExplicitDense", MemberNames, EMemberPresence::RequireAll);
+			FStructSchemaId ImplicitId = Batch.DeclareStruct("Test", "ImplicitDense", MemberNames, EMemberPresence::AllowSparse);
+
+			FMemberBuilder X;
+			X.Add(Batch.NameMember("A"), char8_t('a'));
+			X.Add(Batch.NameMember("B"), char16_t('b'));
+			X.Add(Batch.NameMember("C"), char32_t('c'));
+
+			FMemberBuilder Y;
+			Y.Add(Batch.NameMember("A"), char8_t('1'));
+			Y.Add(Batch.NameMember("B"), char16_t('2'));
+			Y.Add(Batch.NameMember("C"), char32_t('3'));
+						
+			Batch.AddObject(ExplicitId, MoveTemp(X));
+			Batch.AddObject(ImplicitId, MoveTemp(Y));
+		}, 
+		[](TConstArrayView<FStructView> Objects, const FTestNameReader& Names)
+		{
+			CHECK(Objects.Num() == 2);
+			FTestMemberReader X(Objects[0]);
+			FTestMemberReader Y(Objects[1]);
+			
+			CHECK(Names[X.PeekName()]		== "A");
+			CHECK(X.GrabLeaf().AsChar8()	== 'a');
+			CHECK(Names[X.PeekName()]		== "B");
+			CHECK(X.GrabLeaf().AsChar16()	== 'b');
+			CHECK(Names[X.PeekName()]		== "C");
+			CHECK(X.GrabLeaf().AsChar32()	== 'c');
+
+			CHECK(Names[Y.PeekName()]		== "A");
+			CHECK(Y.GrabLeaf().AsChar8()	== '1');
+			CHECK(Names[Y.PeekName()]		== "B");
+			CHECK(Y.GrabLeaf().AsChar16()	== '2');
+			CHECK(Names[Y.PeekName()]		== "C");
+			CHECK(Y.GrabLeaf().AsChar32()	== '3');
+		});
+	}
+
+	SECTION("Struct")
+	{
+		static constexpr std::initializer_list<const char*> ObjectMembers =  {"L1", "S", "N", "L2"};
+		static constexpr std::initializer_list<const char*> StructMembers =  {"Nested", "Leaf"};
+		static constexpr std::initializer_list<const char*> NestedMembers =  {"I1", "I2"};
+
+		TestSerialize([](FTestBatchBuilder& Batch)
+		{
+			FStructSchemaId ObjectId = Batch.DeclareStruct("Test", "Object", ObjectMembers, EMemberPresence::AllowSparse);
+			FStructSchemaId StructId = Batch.DeclareStruct("Test", "Struct", StructMembers, EMemberPresence::AllowSparse);
+			FStructSchemaId NestedId = Batch.DeclareStruct("Test", "Nested", NestedMembers, EMemberPresence::AllowSparse);
+
+			FMemberBuilder Members;
+			Members.Add(	Batch.NameMember("I1"), 100);
+			TUniquePtr<FBuiltStruct> NestedInStruct = Members.BuildAndReset(Batch.Get(NestedId), Batch);
+
+			Members.AddStruct(Batch.NameMember("Nested"), NestedId, MoveTemp(NestedInStruct));
+			Members.Add(Batch.NameMember("Leaf"), true);
+			TUniquePtr<FBuiltStruct> Struct = Members.BuildAndReset(Batch.Get(StructId), Batch);
+		
+			Members.Add(Batch.NameMember("I2"), 200);
+			TUniquePtr<FBuiltStruct> NestedInObject = Members.BuildAndReset(Batch.Get(NestedId), Batch);
+
+			Members.Add(Batch.NameMember("L1"), 123.f);
+			Members.AddStruct(Batch.NameMember("S"), StructId, MoveTemp(Struct));
+			Members.AddStruct(Batch.NameMember("N"), NestedId, MoveTemp(NestedInObject));
+			Members.Add(Batch.NameMember("L2"), -45.f);
+
+			Batch.AddObject(ObjectId, MoveTemp(Members));
+		}, 
+		[](TConstArrayView<FStructView> Objects, const FTestNameReader& Names)
+		{
+			CHECK(Objects.Num() == 1);
+
+			FTestMemberReader Object(Objects[0]);
+			CHECK(Object.GrabLeaf().AsFloat() == 123.f);
+			FTestMemberReader Struct(Object.GrabStruct());
+			FTestMemberReader NestedInObject(Object.GrabStruct());
+			CHECK(Object.GrabLeaf().AsFloat()	== -45.f);
+			
+			FTestMemberReader NestedInStruct(Struct.GrabStruct());
+			CHECK(Struct.GrabLeaf().AsBool() == true);
+
+			CHECK(NestedInObject.GrabLeaf().AsS32() == 200);
+
+			CHECK(NestedInStruct.GrabLeaf().AsS32() == 100);
+		});
+	}
+
+	SECTION("Enum")
+	{
+		static constexpr std::initializer_list<const char*> MemberNames = 
+		{ "A2", "A0", "B0", "B4", "B5", "B7", "C3", "D34", "Max8", "Max16", "Max32", "Max64", "IF" };
+
+		TestSerialize([](FTestBatchBuilder& Batch)
+		{
+			// Test create holes in the original FNameId, FStructSchemaId and FEnumSchemaId index range
+			FStructSchemaId UnusedId = Batch.DeclareStruct("Test", "UnusedStruct", {"U1", "U2"}, EMemberPresence::AllowSparse);
+			FStructSchemaId ObjectId = Batch.DeclareStruct("Test", "Enums", MemberNames, EMemberPresence::AllowSparse);
+
+			FEnumSchemaId U	= Batch.DeclareEnum("Test", "UnusedEnum1",	EEnumMode::Flag, ELeafWidth::B8,	{"U3"}, {1}); // Hole
+			FEnumSchemaId A	= Batch.DeclareEnum("Test", "FlatDense8",	EEnumMode::Flat, ELeafWidth::B8,	{"A", "B", "C"}, {0, 1, 2});
+			FEnumSchemaId X	= Batch.DeclareEnum("Test", "UnusedEnum2",	EEnumMode::Flag, ELeafWidth::B8,	{"U4"}, {1}); // Hole
+			FEnumSchemaId B	= Batch.DeclareEnum("Test", "FlagDense8",	EEnumMode::Flag, ELeafWidth::B8,	{"A", "B", "C"}, {1, 2, 4});
+			FEnumSchemaId C	= Batch.DeclareEnum("Test", "FlatSparse8",	EEnumMode::Flat, ELeafWidth::B8,	{"A", "B", "C"}, {1, 2, 3});
+			FEnumSchemaId D	= Batch.DeclareEnum("Test", "FlagSparse8",	EEnumMode::Flag, ELeafWidth::B8,	{"A", "B", "C"}, {2, 16, 32});
+			FEnumSchemaId E	= Batch.DeclareEnum("Test", "FlatLimit8",	EEnumMode::Flat, ELeafWidth::B8,	{"Min", "Max"}, {0, 0xFF});
+			FEnumSchemaId F	= Batch.DeclareEnum("Test", "FlatLimit16",	EEnumMode::Flat, ELeafWidth::B16,	{"Min", "Max"}, {0, 0xFFFF});
+			FEnumSchemaId G	= Batch.DeclareEnum("Test", "FlatLimit32",	EEnumMode::Flat, ELeafWidth::B32,	{"Min", "Max"}, {0, 0xFFFFFFFF});
+			FEnumSchemaId H	= Batch.DeclareEnum("Test", "FlatLimit64",	EEnumMode::Flat, ELeafWidth::B64,	{"Min", "Max"}, {0, 0xFFFFFFFFFFFFFFFF});
+			FEnumSchemaId I	= Batch.DeclareEnum("Test", "FlagLimit64",	EEnumMode::Flag, ELeafWidth::B64,	{"One", "Max"}, {1, 0x8000000000000000});
+					
+			FMemberBuilder Members;
+			Members.AddEnum8(Batch.NameMember("A2"),		A, 2);
+			Members.AddEnum8(Batch.NameMember("A0"),		A, 0);
+			Members.AddEnum8(Batch.NameMember("B0"),		B, 0);
+			Members.AddEnum8(Batch.NameMember("B4"),		B, 4);
+			Members.AddEnum8(Batch.NameMember("B5"),		B, 5);
+			Members.AddEnum8(Batch.NameMember("B7"),		B, 7);
+			Members.AddEnum8(Batch.NameMember("C3"),		C, 3);
+			Members.AddEnum8(Batch.NameMember("D34"),		D, 34);
+			Members.AddEnum8(Batch.NameMember("Max8"),		E, 0xFF);
+			Members.AddEnum16(Batch.NameMember("Max16"),	F, 0xFFFF);
+			Members.AddEnum32(Batch.NameMember("Max32"),	G, 0xFFFFFFFF);
+			Members.AddEnum64(Batch.NameMember("Max64"),	H, 0xFFFFFFFFFFFFFFFF);
+			Members.AddEnum64(Batch.NameMember("IF"),		I, 0x8000000000000001);
+
+			Batch.AddObject(ObjectId, MoveTemp(Members));
+		}, 
+		[](TConstArrayView<FStructView> Objects, const FTestNameReader& Names)
+		{
+			CHECK(Objects.Num() == 1);
+
+			FReadBatchId Batch = Objects[0].Schema.Batch;
+			auto GetEnumName = [Batch, &Names](FLeafView Leaf) { return Names[ResolveEnumSchema(Batch, Leaf.Enum).Type.Name]; };
+
+			FTestMemberReader It1(Objects[0]);
+			CHECK(It1.GrabLeaf().AsEnum8() == 2);
+			CHECK(It1.GrabLeaf().AsEnum8() == 0);
+			CHECK(It1.GrabLeaf().AsEnum8() == 0);
+			CHECK(It1.GrabLeaf().AsEnum8() == 4);
+			CHECK(It1.GrabLeaf().AsEnum8() == 5);
+			CHECK(It1.GrabLeaf().AsEnum8() == 7);
+			CHECK(It1.GrabLeaf().AsEnum8() == 3);
+			CHECK(It1.GrabLeaf().AsEnum8() == 34);
+			CHECK(It1.GrabLeaf().AsEnum8() == 0xFF);
+			CHECK(It1.GrabLeaf().AsEnum16() == 0xFFFF);
+			CHECK(It1.GrabLeaf().AsEnum32() == 0xFFFFFFFF);
+			CHECK(It1.GrabLeaf().AsEnum64() == 0xFFFFFFFFFFFFFFFF);
+			CHECK(It1.GrabLeaf().AsEnum64() == 0x8000000000000001);	
+
+			FTestMemberReader It2(Objects[0]);
+			CHECK(GetEnumName(It2.GrabLeaf()) == "FlatDense8");
+			CHECK(GetEnumName(It2.GrabLeaf()) == "FlatDense8");
+			CHECK(GetEnumName(It2.GrabLeaf()) == "FlagDense8");
+			CHECK(GetEnumName(It2.GrabLeaf()) == "FlagDense8");
+			CHECK(GetEnumName(It2.GrabLeaf()) == "FlagDense8");
+			CHECK(GetEnumName(It2.GrabLeaf()) == "FlagDense8");
+			CHECK(GetEnumName(It2.GrabLeaf()) == "FlatSparse8");
+			CHECK(GetEnumName(It2.GrabLeaf()) == "FlagSparse8");
+			CHECK(GetEnumName(It2.GrabLeaf()) == "FlatLimit8");
+			CHECK(GetEnumName(It2.GrabLeaf()) == "FlatLimit16");
+			CHECK(GetEnumName(It2.GrabLeaf()) == "FlatLimit32");
+			CHECK(GetEnumName(It2.GrabLeaf()) == "FlatLimit64");
+			CHECK(GetEnumName(It2.GrabLeaf()) == "FlagLimit64");			
+		});
+	}
+
+	SECTION("LeafRange")
+	{
+		enum class EABCD : uint16 { A, B, C, D };
+		enum class EUnused1 : uint8 { X };
+		enum class EUnused2 : uint8 { Y };
+
+		TestSerialize([](FTestBatchBuilder& Batch)
+		{
+			static constexpr std::initializer_list<const char*> MemberNames = { "B0", "B1", "B8", "B9", "D0", "D3", "Hi", "E3", "E0" };
+		
+			FStructSchemaId ObjectId = Batch.DeclareStruct("Test", "Object", MemberNames, EMemberPresence::AllowSparse);
+			FEnumSchemaId Enum	= Batch.DeclareEnum("Test", "ABCD", EEnumMode::Flat, ELeafWidth::B16, {"A", "B", "C", "D"}, {0, 1, 2, 3});
+			FEnumSchemaId UnusedEnum1	= Batch.DeclareEnum("Test", "Unused1", EEnumMode::Flat, ELeafWidth::B8, {"X"}, {0});
+			FEnumSchemaId UnusedEnum2	= Batch.DeclareEnum("Test", "Unused2", EEnumMode::Flat, ELeafWidth::B8, {"Y"}, {0});
+		
+			FMemberBuilder Members;
+			Members.AddRange(Batch.NameMember("B0"), BuildLeafRange(TConstArrayView<bool>()));
+			Members.AddRange(Batch.NameMember("B1"), BuildLeafRange(MakeArrayView({true})));
+			Members.AddRange(Batch.NameMember("B8"), BuildLeafRange(MakeArrayView({false, true, false, true, false, true, false, true})));
+			Members.AddRange(Batch.NameMember("B9"), BuildLeafRange(MakeArrayView({true, false, true, false, true, false, true, false, true})));
+			Members.AddRange(Batch.NameMember("D0"), BuildLeafRange(TConstArrayView<double>()));
+			Members.AddRange(Batch.NameMember("D3"), BuildLeafRange(MakeArrayView({DBL_MIN, 0.0, DBL_MAX})));
+			Members.AddRange(Batch.NameMember("Hi"), BuildLeafRange(MakeArrayView(u8"Hello!")));
+			Members.AddRange(Batch.NameMember("E3"), BuildEnumRange(Enum, MakeArrayView({EABCD::B, EABCD::A, EABCD::D})));
+			Members.AddRange(Batch.NameMember("E0"), BuildEnumRange(Enum, TConstArrayView<EUnused1>()));
+
+			Batch.AddObject(ObjectId, MoveTemp(Members));
+		}, 
+		[](TConstArrayView<FStructView> Objects, const FTestNameReader& Names)
+		{
+			CHECK(Objects.Num() == 1);
+
+			FTestMemberReader It(Objects[0]);
+			FLeafRangeView B0 = It.GrabRange().AsLeaves();
+			FLeafRangeView B1 = It.GrabRange().AsLeaves();
+			FLeafRangeView B8 = It.GrabRange().AsLeaves();
+			FLeafRangeView B9 = It.GrabRange().AsLeaves();
+			FLeafRangeView D0 = It.GrabRange().AsLeaves();
+			FLeafRangeView D3 = It.GrabRange().AsLeaves();
+			FLeafRangeView Hi = It.GrabRange().AsLeaves();
+			FLeafRangeView E3 = It.GrabRange().AsLeaves();
+			FLeafRangeView E0 = It.GrabRange().AsLeaves();
+			
+			CHECK(B0.Num() == 0);
+			CHECK(EqualItems(B1.AsBools(),	MakeArrayView({true})));
+			CHECK(EqualItems(B8.AsBools(),	MakeArrayView({false, true, false, true, false, true, false, true})));
+			CHECK(EqualItems(B9.AsBools(),	MakeArrayView({true, false, true, false, true, false, true, false, true})));
+			CHECK(EqualItems(D0.AsDoubles(), TConstArrayView<double>()));
+			CHECK(EqualItems(D3.AsDoubles(), MakeArrayView({DBL_MIN, 0.0, DBL_MAX})));
+			CHECK(EqualItems(Hi.AsUtf8(), u8"Hello!"));
+			CHECK(EqualItems(E3.As<EABCD>(), MakeArrayView({EABCD::B, EABCD::A, EABCD::D})));
+			CHECK(EqualItems(E0.As<EUnused1>(), TConstArrayView<EUnused1>()));
+		});
+	}
+		
+	SECTION("StructRange")
+	{
+		TestSerialize([](FTestBatchBuilder& Batch)
+		{
+			FStructSchemaId ObjectId = Batch.DeclareStruct("Test", "Object", {"Structs"}, EMemberPresence::AllowSparse);
+			FStructSchemaId StructId = Batch.DeclareStruct("Test", "Struct", {"I", "F"}, EMemberPresence::AllowSparse);
+			 
+			FStructRangeBuilder Structs(3);
+			Structs[0].Add(Batch.NameMember("I"), 1);
+			Structs[1].Add(Batch.NameMember("F"), 1.f);
+		
+			FMemberBuilder Members;
+			Members.AddRange(Batch.NameMember("Structs"), Structs.BuildAndReset(Batch.Get(StructId), Batch));
+
+			Batch.AddObject(ObjectId, MoveTemp(Members));
+		}, 
+		[](TConstArrayView<FStructView> Objects, const FTestNameReader& Names)
+		{
+			CHECK(Objects.Num() == 1);
+
+			FTestMemberReader It(Objects[0]);
+			TArray<FTestMemberReader> Structs = MakeArray<FTestMemberReader>(It.GrabRange().AsStructs());
+			CHECK(Structs.Num() == 3);
+			CHECK(Structs[0].GrabLeaf().AsS32() == 1);
+			CHECK(Structs[1].GrabLeaf().AsFloat() == 1.f);	
+		});
+	}
+	
+	SECTION("NestedRange")
+	{
+		enum class EAB : uint8 { A = 1, B = 4 };
+
+		TestSerialize([](FTestBatchBuilder& Batch)
+		{
+			FStructSchemaId Object = Batch.DeclareStruct("Test", "Object", {"IntRs", "EmptyRs", "EnumRs", "StructRs", "StructRRs"}, EMemberPresence::AllowSparse);
+			FStructSchemaId XY = Batch.DeclareStruct("Test", "XY", {"X", "Y"}, EMemberPresence::RequireAll);
+			FStructSchemaId ZW = Batch.DeclareStruct("Test", "ZW", {"Z", "W"}, EMemberPresence::AllowSparse);
+			FEnumSchemaId Enum	= Batch.DeclareEnum("Test", "AB", EEnumMode::Flag, ELeafWidth::B8, {"A", "B"}, {1, 4});
+	
+			FNestedRangeBuilder IntRs(MakeLeafRangeSchema<int32, int32>(), 3);
+			IntRs.Add(BuildLeafRange(MakeArrayView({1})));
+			IntRs.Add({});
+			IntRs.Add(BuildLeafRange(MakeArrayView({2, 3})));
+
+			FNestedRangeBuilder EnumRs(MakeEnumRangeSchema<EAB, int32>(Enum), 2);
+			EnumRs.Add({});
+			EnumRs.Add(BuildEnumRange(Enum, MakeArrayView({EAB::A, EAB(0), EAB::B})));
+
+			FStructRangeBuilder XYs(uint64(2));
+			XYs[0].Add(Batch.NameMember("X"), 1.f);
+			XYs[0].Add(Batch.NameMember("Y"), 2.f);
+			XYs[1].Add(Batch.NameMember("X"), 3.f);
+			XYs[1].Add(Batch.NameMember("Y"), 4.f);
+			FNestedRangeBuilder StructRs(MakeStructRangeSchema(ERangeSizeType::U64, XY), 1);
+			StructRs.Add(XYs.BuildAndReset(Batch.Get(XY), Batch));
+
+			FStructRangeBuilder ZWs(int16(3));
+			ZWs[0].Add(Batch.NameMember("Z"), 1.5f);
+			ZWs[2].Add(Batch.NameMember("Z"), 2.5f);
+			ZWs[2].Add(Batch.NameMember("W"), 3.5f);
+			FMemberSchema ZWRangeSchema = MakeStructRangeSchema(ERangeSizeType::S16, ZW);
+			FNestedRangeBuilder ZWRs(ZWRangeSchema, 1);
+			ZWRs.Add(ZWs.BuildAndReset(Batch.Get(ZW), Batch));
+			FNestedRangeBuilder StructRRs(MakeNestedRangeSchema(ERangeSizeType::U32, ZWRangeSchema), 1);
+			StructRRs.Add(ZWRs.BuildAndReset(ERangeSizeType::U32));
+
+
+			FMemberBuilder Members;
+			Members.AddRange(Batch.NameMember("IntRs"), IntRs.BuildAndReset(ERangeSizeType::S32));
+			Members.AddRange(Batch.NameMember("EmptyRs"), IntRs.BuildAndReset(ERangeSizeType::S32));
+			Members.AddRange(Batch.NameMember("EnumRs"), EnumRs.BuildAndReset(ERangeSizeType::U8));
+			Members.AddRange(Batch.NameMember("StructRs"), StructRs.BuildAndReset(ERangeSizeType::U64));
+			Members.AddRange(Batch.NameMember("StructRRs"), StructRRs.BuildAndReset(ERangeSizeType::U32));
+
+			Batch.AddObject(Object, MoveTemp(Members));
+		}, 
+		[](TConstArrayView<FStructView> Objects, const FTestNameReader& Names)
+		{
+			CHECK(Objects.Num() == 1);
+
+			FTestMemberReader It(Objects[0]);
+			TArray<FRangeView> IntRs = MakeArray<FRangeView>(It.GrabRange().AsRanges());
+			FNestedRangeView EmptyRs = It.GrabRange().AsRanges();
+			TArray<FRangeView> EnumRs = MakeArray<FRangeView>(It.GrabRange().AsRanges());
+			TArray<FRangeView> StructRs = MakeArray<FRangeView>(It.GrabRange().AsRanges());
+			TArray<FRangeView> StructRRs = MakeArray<FRangeView>(It.GrabRange().AsRanges());
+
+			CHECK(IntRs.Num() == 3);
+			CHECK(EqualItems(IntRs[0].AsLeaves().AsS32s(), MakeArrayView({1})));
+			CHECK(IntRs[1].IsEmpty());
+			CHECK(EqualItems(IntRs[2].AsLeaves().AsS32s(), MakeArrayView({2, 3})));
+
+			CHECK(EmptyRs.Num() == 0);
+			
+			CHECK(EnumRs.Num() == 2);
+			CHECK(EnumRs[0].IsEmpty());
+			CHECK(EqualItems(EnumRs[1].AsLeaves().As<EAB>(), MakeArrayView({EAB::A, EAB(0), EAB::B})));
+
+			CHECK(StructRs.Num() == 1);
+			TArray<FTestMemberReader> XYs = MakeArray<FTestMemberReader>(StructRs[0].AsStructs());
+			CHECK(Names[XYs[0].PeekName()] == "X");
+			CHECK(XYs[0].GrabLeaf().AsFloat() == 1.f);
+			CHECK(Names[XYs[0].PeekName()] == "Y");
+			CHECK(XYs[0].GrabLeaf().AsFloat() == 2.f);
+			CHECK(Names[XYs[1].PeekName()] == "X");
+			CHECK(XYs[1].GrabLeaf().AsFloat() == 3.f);
+			CHECK(Names[XYs[1].PeekName()] == "Y");
+			CHECK(XYs[1].GrabLeaf().AsFloat() == 4.f);
+
+			CHECK(StructRRs.Num() == 1);
+			TArray<FRangeView> ZWRs = MakeArray<FRangeView>(StructRRs[0].AsRanges());
+			CHECK(ZWRs.Num() == 1);
+			TArray<FTestMemberReader> ZWs = MakeArray<FTestMemberReader>(ZWRs[0].AsStructs());
+			CHECK(ZWs.Num() == 3);
+			CHECK(Names[ZWs[0].PeekName()] == "Z");
+			CHECK(ZWs[0].GrabLeaf().AsFloat() == 1.5f);
+			CHECK(Names[ZWs[2].PeekName()] == "Z");
+			CHECK(ZWs[2].GrabLeaf().AsFloat() == 2.5f);
+			CHECK(Names[ZWs[2].PeekName()] == "W");
+			CHECK(ZWs[2].GrabLeaf().AsFloat() == 3.5f);
+		});
+
+	}
+	
+	SECTION("UniRange")
+	{
+		TestSerialize([](FTestBatchBuilder& Batch)
+		{
+			FStructSchemaId Object = Batch.DeclareStruct("Test", "Object", {"Bools", "Structs", "BF", "BT" }, EMemberPresence::AllowSparse);
+			FStructSchemaId Struct = Batch.DeclareStruct("Test", "Struct", {"MaybeB", "Bs", "MaybeBs", "B"}, EMemberPresence::AllowSparse);
+						
+			const bool True = true;
+			const bool False = false;
+			FNestedRangeBuilder MaybeBs(MakeLeafRangeSchema<bool, bool>(), 1);
+			FStructRangeBuilder Structs(10);
+			Structs[5].AddRange(Batch.NameMember("MaybeB"), BuildLeafRange(&False, true));
+			Structs[6].AddRange(Batch.NameMember("MaybeB"), BuildLeafRange(&True, false));
+			Structs[7].AddRange(Batch.NameMember("MaybeB"), BuildLeafRange(&True, true));
+			Structs[7].AddRange(Batch.NameMember("Bs"),		BuildLeafRange(MakeArrayView({true, true, false, false, true, true, false, false, true, true})));
+			MaybeBs.Add(BuildLeafRange(&True, true));
+			Structs[7].AddRange(Batch.NameMember("MaybeBs"), MaybeBs.BuildAndReset(ERangeSizeType::Uni));
+			Structs[7].Add(Batch.NameMember("B"), true);
+			MaybeBs.Add(BuildLeafRange(&True, false));
+			Structs[8].AddRange(Batch.NameMember("MaybeBs"), MaybeBs.BuildAndReset(ERangeSizeType::Uni));
+			Structs[9].Add(Batch.NameMember("B"), false);
+
+			FMemberBuilder Members;
+			Members.AddRange(Batch.NameMember("Bools"), BuildLeafRange(&True, true));
+			Members.AddRange(Batch.NameMember("Structs"), Structs.BuildAndReset(Batch.Get(Struct), Batch));
+			Members.Add(Batch.NameMember("BF"), false);
+			Members.Add(Batch.NameMember("BT"), true);
+
+			Batch.AddObject(Object, MoveTemp(Members));
+		}, 
+		[](TConstArrayView<FStructView> Objects, const FTestNameReader& Names)
+		{
+			CHECK(Objects.Num() == 1);
+			FTestMemberReader It(Objects[0]);
+			
+			FBoolRangeView Bools = It.GrabRange().AsLeaves().AsBools();
+			TArray<FTestMemberReader> Structs = MakeArray<FTestMemberReader>(It.GrabRange().AsStructs());
+			CHECK(It.GrabLeaf().AsBool() == false);
+			CHECK(It.GrabLeaf().AsBool() == true);
+
+			CHECK(Bools.Num() == 1);
+			CHECK(Bools[0] == true);
+			
+			CHECK(EqualItems(Structs[5].GrabRange().AsLeaves().AsBools(), MakeArrayView({false})));
+			CHECK(Structs[6].GrabRange().AsLeaves().AsBools().Num() == 0);
+			CHECK(EqualItems(Structs[7].GrabRange().AsLeaves().AsBools(), MakeArrayView({true})));
+			CHECK(EqualItems(Structs[7].GrabRange().AsLeaves().AsBools(), MakeArrayView({true, true, false, false, true, true, false, false, true, true})));
+			TArray<FRangeView> MaybeBs7 = MakeArray<FRangeView>(Structs[7].GrabRange().AsRanges());
+			CHECK(MaybeBs7.Num() == 1);
+			CHECK(EqualItems(MaybeBs7[0].AsLeaves().AsBools(), MakeArrayView({true})));
+			CHECK(Structs[7].GrabLeaf().AsBool() == true);
+			TArray<FRangeView> MaybeBs8 = MakeArray<FRangeView>(Structs[8].GrabRange().AsRanges());
+			CHECK(MaybeBs8.Num() == 1);
+			CHECK(MaybeBs8[0].AsLeaves().AsBools().Num() == 0);
+			CHECK(Structs[9].GrabLeaf().AsBool() == false);
+		});
+	}
+	
+	SECTION("DynamicStruct")
+	{
+		TestSerialize([](FTestBatchBuilder& Batch)
+		{
+			FStructSchemaId Unused1 = Batch.DeclareStruct("Test", "Unused1", {"X"}, EMemberPresence::AllowSparse);
+			FStructSchemaId SA = Batch.DeclareStruct("Test", "SA", {"X"}, EMemberPresence::AllowSparse);
+			FStructSchemaId Unused2 = Batch.DeclareStruct("Test", "Unused2", {"X"}, EMemberPresence::AllowSparse);
+			FStructSchemaId SB = Batch.DeclareStruct("Test", "SB", {"X"}, EMemberPresence::AllowSparse);
+			FStructSchemaId Object = Batch.DeclareStruct("Test", "Object", {"Same", "Some", "None", "Diff"}, EMemberPresence::AllowSparse);
+			FStructSchemaId Unused3 = Batch.DeclareStruct("Test", "Unused3", {"X"}, EMemberPresence::AllowSparse);
+		
+			auto BuildStruct = [&Batch](FStructSchemaId Struct, auto X)
+				{
+					FMemberBuilder Members;
+					Members.Add(Batch.NameMember("X"), X);
+					return Members.BuildAndReset(Batch.Get(Struct), Batch);
+				};
+
+			FMemberBuilder O1;
+			O1.AddStruct(Batch.NameMember("Same"), SA, BuildStruct(SA, 0));
+			O1.AddStruct(Batch.NameMember("Some"), SA, BuildStruct(SA, 1));
+			O1.AddStruct(Batch.NameMember("Diff"), SA, BuildStruct(SA, 2));
+			FMemberBuilder O2;
+			O2.AddStruct(Batch.NameMember("Same"), SA, BuildStruct(SA, 3));
+			O2.AddStruct(Batch.NameMember("Diff"), SB, BuildStruct(SB, 4.f));
+			
+			Batch.AddObject(Object, MoveTemp(O1));
+			Batch.AddObject(Object, MoveTemp(O2));
+		}, 
+		[](TConstArrayView<FStructView> Objects, const FTestNameReader& Names)
+		{
+			CHECK(Objects.Num() == 2);
+
+			FTestMemberReader O1(Objects[0]);
+			CHECK(O1.PeekType().AsStruct().IsDynamic == 0);
+			CHECK(FTestMemberReader(O1.GrabStruct()).GrabLeaf().AsS32() == 0);
+			CHECK(O1.PeekType().AsStruct().IsDynamic == 0);
+			CHECK(FTestMemberReader(O1.GrabStruct()).GrabLeaf().AsS32() == 1);
+			CHECK(O1.PeekType().AsStruct().IsDynamic == 1);
+			CHECK(FTestMemberReader(O1.GrabStruct()).GrabLeaf().AsS32() == 2);
+			
+			FTestMemberReader O2(Objects[1]);
+			CHECK(O2.PeekType().AsStruct().IsDynamic == 0);
+			CHECK(FTestMemberReader(O2.GrabStruct()).GrabLeaf().AsS32() == 3);
+			CHECK(O2.PeekType().AsStruct().IsDynamic == 1);
+			CHECK(FTestMemberReader(O2.GrabStruct()).GrabLeaf().AsFloat() == 4.f);
+		});
+	}
+
+	SECTION("DynamicStructRange")
+	{
+		TestSerialize([](FTestBatchBuilder& Batch)
+		{
+			FStructSchemaId SA = Batch.DeclareStruct("Test", "SA", {"X"}, EMemberPresence::AllowSparse);
+			FStructSchemaId Unused = Batch.DeclareStruct("Test", "Unused2", {"X"}, EMemberPresence::AllowSparse);
+			FStructSchemaId SB = Batch.DeclareStruct("Test", "SB", {"X"}, EMemberPresence::AllowSparse);
+			FStructSchemaId Object = Batch.DeclareStruct("Test", "Object", {"Same", "Some", "None", "Diff", "SameEmpty", "DiffEmpty", "DiffNested"}, EMemberPresence::AllowSparse);
+		
+			auto BuildStructRange = [&Batch](FStructSchemaId Struct, auto X)
+				{
+					FStructRangeBuilder Members(1);
+					Members[0].Add(Batch.NameMember("X"), X);
+					return Members.BuildAndReset(Batch.Get(Struct), Batch);
+				};
+
+			
+			FMemberBuilder O1;
+			O1.AddRange(Batch.NameMember("Same"), BuildStructRange(SA, 10));
+			O1.AddRange(Batch.NameMember("Some"), BuildStructRange(SA, 11));
+			O1.AddRange(Batch.NameMember("Diff"), BuildStructRange(SA, 12));
+			O1.AddRange(Batch.NameMember("SameEmpty"), BuildStructRange(SA, 13));
+			O1.AddRange(Batch.NameMember("DiffEmpty"), BuildStructRange(SA, 14));
+			FNestedRangeBuilder NestedSA(MakeStructRangeSchema(ERangeSizeType::S32, SA), 1);
+			NestedSA.Add(BuildStructRange(SA, 100));
+			O1.AddRange(Batch.NameMember("DiffNested"), NestedSA.BuildAndReset(ERangeSizeType::S32));
+			
+			FMemberBuilder O2;
+			O2.AddRange(Batch.NameMember("Same"), BuildStructRange(SA, 20));
+			O2.AddRange(Batch.NameMember("Diff"), BuildStructRange(SB, 22.f));
+			O2.AddRange(Batch.NameMember("SameEmpty"), FStructRangeBuilder(0).BuildAndReset(Batch.Get(SA), Batch));
+			O2.AddRange(Batch.NameMember("DiffEmpty"), FStructRangeBuilder(0).BuildAndReset(Batch.Get(SB), Batch));
+			FNestedRangeBuilder NestedSB(MakeStructRangeSchema(ERangeSizeType::S32, SB), 1);
+			NestedSB.Add(BuildStructRange(SB, 200.f));
+			O2.AddRange(Batch.NameMember("DiffNested"), NestedSB.BuildAndReset(ERangeSizeType::S32));
+			
+			Batch.AddObject(Object, MoveTemp(O1));
+			Batch.AddObject(Object, MoveTemp(O2));
+		}, 
+		[](TConstArrayView<FStructView> Objects, const FTestNameReader& Names)
+		{
+			CHECK(Objects.Num() == 2);
+
+			FTestMemberReader O1(Objects[0]);
+			CHECK(MakeArray<FTestMemberReader>(O1.GrabRange().AsStructs())[0].GrabLeaf().AsS32() == 10);
+			CHECK(MakeArray<FTestMemberReader>(O1.GrabRange().AsStructs())[0].GrabLeaf().AsS32() == 11);
+			CHECK(MakeArray<FTestMemberReader>(O1.GrabRange().AsStructs())[0].GrabLeaf().AsS32() == 12);
+			CHECK(MakeArray<FTestMemberReader>(O1.GrabRange().AsStructs())[0].GrabLeaf().AsS32() == 13);
+			CHECK(MakeArray<FTestMemberReader>(O1.GrabRange().AsStructs())[0].GrabLeaf().AsS32() == 14);
+			TArray<FRangeView> DiffNested1 = MakeArray<FRangeView>(O1.GrabRange().AsRanges());
+			CHECK(MakeArray<FTestMemberReader>(DiffNested1[0].AsStructs())[0].GrabLeaf().AsS32() == 100);
+
+			FTestMemberReader O2(Objects[1]);
+			CHECK(MakeArray<FTestMemberReader>(O2.GrabRange().AsStructs())[0].GrabLeaf().AsS32() == 20);
+			CHECK(MakeArray<FTestMemberReader>(O2.GrabRange().AsStructs())[0].GrabLeaf().AsFloat() == 22.f);
+			CHECK(O2.GrabRange().AsStructs().Num() == 0);
+			CHECK(O2.GrabRange().AsStructs().Num() == 0);
+			TArray<FRangeView> DiffNested2 = MakeArray<FRangeView>(O2.GrabRange().AsRanges());
+			CHECK(MakeArray<FTestMemberReader>(DiffNested2[0].AsStructs())[0].GrabLeaf().AsFloat() == 200.f);
+		});
+	}
+
+	SECTION("Inheritance")
+	{
+		TestSerialize([](FTestBatchBuilder& Batch)
+		{
+			FStructSchemaId Unused = Batch.DeclareStruct("Test", "X", {"X"}, EMemberPresence::AllowSparse);
+			FStructSchemaId Low = Batch.DeclareStruct("Test", "Low", {"LInt"}, EMemberPresence::AllowSparse);
+			FStructSchemaId Mid = Batch.DeclareStruct("Test", "Mid", {"MInt", "MLow"}, EMemberPresence::AllowSparse, ToOptional(Low));
+			FStructSchemaId Top = Batch.DeclareStruct("Test", "Top", {"TInt", "TLow", "TMids"}, EMemberPresence::AllowSparse, ToOptional(Mid));
+			
+			FMemberBuilder Members;
+			Members.Add(Batch.NameMember("LInt"), 123);
+			Members.BuildSuperStruct(Batch.Get(Low), Batch);
+			Members.Add(Batch.NameMember("MInt"), 456);
+			FMemberBuilder Nested;
+			Nested.Add(Batch.NameMember("LInt"), 1000);
+			Members.AddStruct(Batch.NameMember("MLow"), Low, Nested.BuildAndReset(Batch.Get(Low), Batch));
+			Members.BuildSuperStruct(Batch.Get(Mid), Batch);
+			Members.Add(Batch.NameMember("TInt"), 789);
+			Nested.Add(Batch.NameMember("LInt"), 2000);
+			Members.AddStruct(Batch.NameMember("TLow"), Low, Nested.BuildAndReset(Batch.Get(Low), Batch));
+			FStructRangeBuilder NestedRange(1);
+			NestedRange[0].Add(Batch.NameMember("MInt"), 3000);
+			Members.AddRange(Batch.NameMember("TMids"), NestedRange.BuildAndReset(Batch.Get(Mid), Batch));
+
+			Batch.AddObject(Top, MoveTemp(Members));
+		}, 
+		[](TConstArrayView<FStructView> Objects, const FTestNameReader& Names)
+		{
+			CHECK(Objects.Num() == 1);
+
+			FTestMemberReader TopIt(Objects[0]);
+			FTestMemberReader MidIt(TopIt.GrabStruct());
+			FTestMemberReader LowIt(MidIt.GrabStruct());
+			CHECK(LowIt.GrabLeaf().AsS32() == 123);
+			CHECK(MidIt.GrabLeaf().AsS32() == 456);
+			CHECK(FTestMemberReader(MidIt.GrabStruct()).GrabLeaf().AsS32() == 1000);
+			CHECK(Names[TopIt.PeekName()] == "TInt");
+			CHECK(TopIt.GrabLeaf().AsS32() == 789);
+			CHECK(Names[TopIt.PeekName()] == "TLow");
+			CHECK(FTestMemberReader(TopIt.GrabStruct()).GrabLeaf().AsS32() == 2000);
+			CHECK(Names[TopIt.PeekName()] == "TMids");
+			CHECK(MakeArray<FTestMemberReader>(TopIt.GrabRange().AsStructs())[0].GrabLeaf().AsS32() == 3000);
+
+			FFlatMemberReader FlatIt(Objects[0]);
+			CHECK(Names[FlatIt.PeekOwner().Name] == "Low");
+			CHECK(FlatIt.GrabLeaf().AsS32() == 123);
+			CHECK(Names[FlatIt.PeekOwner().Name] == "Mid");
+			CHECK(FlatIt.GrabLeaf().AsS32() == 456);
+			CHECK(Names[FlatIt.PeekOwner().Name] == "Mid");
+			CHECK(FTestMemberReader(FlatIt.GrabStruct()).GrabLeaf().AsS32() == 1000);
+			CHECK(Names[FlatIt.PeekOwner().Name] == "Top");
+			CHECK(FlatIt.GrabLeaf().AsS32() == 789);
+			CHECK(Names[FlatIt.PeekOwner().Name] == "Top");
+			CHECK(FTestMemberReader(FlatIt.GrabStruct()).GrabLeaf().AsS32() == 2000);
+			CHECK(Names[FlatIt.PeekOwner().Name] == "Top");
+			CHECK(MakeArray<FTestMemberReader>(FlatIt.GrabRange().AsStructs())[0].GrabLeaf().AsS32() == 3000);
+			CHECK(!FlatIt.HasMore());
+		});
+	}
+
+	SECTION("SparseInheritance")
+	{
+		TestSerialize([](FTestBatchBuilder& Batch)
+		{																										// Usage by A B C 
+			FStructSchemaId B0 = Batch.DeclareStruct("Test", "B0", {"0"}, EMemberPresence::AllowSparse);				 // - - -
+			FStructSchemaId B1 = Batch.DeclareStruct("Test", "B1", {"1"}, EMemberPresence::AllowSparse, ToOptional(B0)); // 1 - 0
+			FStructSchemaId B2 = Batch.DeclareStruct("Test", "B2", {"2"}, EMemberPresence::AllowSparse, ToOptional(B1)); // - 1 1
+			FStructSchemaId B3 = Batch.DeclareStruct("Test", "B3", {"3"}, EMemberPresence::AllowSparse, ToOptional(B2)); // - - 0
+			FStructSchemaId B4 = Batch.DeclareStruct("Test", "B4", {"4"}, EMemberPresence::AllowSparse, ToOptional(B3)); // 1 1 0
+			FStructSchemaId B5 = Batch.DeclareStruct("Test", "B5", {"5"}, EMemberPresence::AllowSparse, ToOptional(B4)); // 1 0 0
+			FStructSchemaId B6 = Batch.DeclareStruct("Test", "B6", {"6"}, EMemberPresence::AllowSparse, ToOptional(B5)); // 0 - 0
+			FStructSchemaId C5 = Batch.DeclareStruct("Test", "C5", {"5"}, EMemberPresence::AllowSparse, ToOptional(B4)); // - - -
+			
+			FMemberBuilder A;
+			A.Add(Batch.NameMember("1"), 1);
+			A.BuildSuperStruct(Batch.Get(B1), Batch);
+			A.Add(Batch.NameMember("4"), 4);
+			A.BuildSuperStruct(Batch.Get(B4), Batch);
+			A.Add(Batch.NameMember("5"), 5);
+			A.BuildSuperStruct(Batch.Get(B5), Batch);
+
+			FMemberBuilder B;
+			B.Add(Batch.NameMember("2"), 20);
+			B.BuildSuperStruct(Batch.Get(B2), Batch);
+			B.Add(Batch.NameMember("4"), 40);
+			B.BuildSuperStruct(Batch.Get(B4), Batch);
+
+			FMemberBuilder C;
+			C.BuildSuperStruct(Batch.Get(B1), Batch); // Empty -> noop
+			C.Add(Batch.NameMember("2"), 200);
+			C.BuildSuperStruct(Batch.Get(B2), Batch);
+			C.BuildSuperStruct(Batch.Get(B3), Batch); // Empty -> noop
+			C.Add(Batch.NameMember("4"), 400);
+			C.BuildSuperStruct(Batch.Get(B4), Batch);
+			C.BuildSuperStruct(Batch.Get(B5), Batch); // Empty -> noop
+			
+			Batch.AddObject(B6, MoveTemp(A));
+			Batch.AddObject(B5, MoveTemp(B));
+			Batch.AddObject(B6, MoveTemp(C));
+		}, 
+		[](TConstArrayView<FStructView> Objects, const FTestNameReader& Names)
+		{
+			CHECK(Objects.Num() == 3);
+			
+			const FStructSchema& Schema0 = Objects[0].Schema.Resolve();
+			const FStructSchema& Schema1 = Objects[1].Schema.Resolve();
+			const FStructSchema& Schema2 = Objects[2].Schema.Resolve();
+			CHECK(Names[Schema0.Type.Name] == "B6");
+			CHECK(Names[Schema1.Type.Name] == "B5");
+			CHECK(Names[Schema2.Type.Name] == "B6");
+			CHECK(Schema0.GetSuperSchema() == Objects[1].Schema.Id);
+			CHECK(Schema2.GetSuperSchema() == Objects[1].Schema.Id);
+			CHECK(&Objects[0].Schema.ResolveSuper() == &Schema1);
+			
+			CHECK(Names[FMemberReader(Objects[0]).GrabStruct().Schema.Resolve().Type.Name]  == "B5");
+			CHECK(Names[FMemberReader(Objects[1]).GrabStruct().Schema.Resolve().Type.Name]  == "B4");
+			CHECK(Names[FMemberReader(Objects[2]).GrabStruct().Schema.Resolve().Type.Name]  == "B4");
+
+			FTestMemberReader B1(Objects[1]);
+			FStructView Super = B1.GrabStruct();
+
+			FFlatMemberReader A(Objects[0]);
+			FFlatMemberReader B(Objects[1]);
+			FFlatMemberReader C(Objects[2]);
+			CHECK(A.GrabLeaf().AsS32() == 1);
+			CHECK(A.GrabLeaf().AsS32() == 4);
+			CHECK(A.GrabLeaf().AsS32() == 5);
+			CHECK(B.GrabLeaf().AsS32() == 20);
+			CHECK(B.GrabLeaf().AsS32() == 40);
+			CHECK(C.GrabLeaf().AsS32() == 200);
+			CHECK(C.GrabLeaf().AsS32() == 400);
+			CHECK(!A.HasMore());
+			CHECK(!B.HasMore());
+			CHECK(!C.HasMore());
+		});
+	}
+
+	SECTION("SparseIndex")
+	{
+		TestSerialize([](FTestBatchBuilder& Batch)
+		{
+			FScopeId Unused = Batch.MakeScope("Unused");
+			FScopeId NestedUnused1 = Batch.NestScope(Unused, "NestedUnused1");
+			FScopeId FlatUsed = Batch.MakeScope("FlatUsed");
+			FScopeId NestedUsed = Batch.NestScope(FlatUsed, "NestedUsed");
+			FScopeId NestedUnused2 = Batch.NestScope(Unused, "NestedUnused2");
+			FScopeId DoubleNested = Batch.NestScope(NestedUsed, "DoubleNested");
+			FScopeId NestedUnused3 = Batch.NestScope(FlatUsed, "NestedUnused3");
+			
+			FTypeId E1T{NestedUnused1,	Batch.MakeTypename("E1")};
+			FTypeId E2T{NestedUsed,		Batch.MakeTypename("E2")};
+			FTypeId E3T{NestedUnused2,	Batch.MakeTypename("E3")};
+			
+			FEnumSchemaId E1D = Batch.DeclareEnum(E1T, EEnumMode::Flat, ELeafWidth::B8, {"C1"}, {1});
+			FEnumSchemaId E2D = Batch.DeclareEnum(E2T, EEnumMode::Flat, ELeafWidth::B8, {"C2"}, {2});
+			FEnumSchemaId E3D = Batch.DeclareEnum(E3T, EEnumMode::Flat, ELeafWidth::B8, {"C3"}, {3});
+
+			FTypeId S1T{NestedUnused1,	Batch.MakeTypename("S1")};
+			FTypeId S2T{NestedUsed,		Batch.MakeTypename("S2")};
+			FTypeId S3T = Batch.MakeParametricType({NestedUnused2,	Batch.MakeTypename("S3")}, {S1T});
+			FTypeId S4T = Batch.MakeParametricType({DoubleNested,	Batch.MakeTypename("S4")}, {S2T, E2T});
+			FTypeId S5T = Batch.MakeParametricType({NestedUnused3,	Batch.MakeTypename("S5")}, {E3T, E1T, S2T});
+			
+			FStructSchemaId S1D = Batch.DeclareStruct(S1T, {"M1"}, EMemberPresence::AllowSparse);
+			FStructSchemaId S2D = Batch.DeclareStruct(S2T, {"M2"}, EMemberPresence::AllowSparse);
+			FStructSchemaId S3D = Batch.DeclareStruct(S3T, {"M3"}, EMemberPresence::AllowSparse);
+			FStructSchemaId S4D = Batch.DeclareStruct(S4T, {"M4"}, EMemberPresence::AllowSparse);
+			FStructSchemaId S5D = Batch.DeclareStruct(S5T, {"M5"}, EMemberPresence::AllowSparse);
+
+			FMemberBuilder Members;
+			Members.Add(Batch.NameMember("M4"), 1);
+
+			Batch.AddObject(S4D, MoveTemp(Members));
+		}, 
+		[](TConstArrayView<FStructView> Objects, const FTestNameReader& Names)
+		{	
+			FReadBatchId Batch = Objects[0].Schema.Batch;
+			FTypeId S4T = Objects[0].Schema.Resolve().Type;
+
+			FNestedScope DoubleNested = ResolveUntranslatedNestedScope(Batch, S4T.Scope.AsNested());
+			FNestedScope NestedUsed = ResolveUntranslatedNestedScope(Batch, DoubleNested.Outer.AsNested());
+			FFlatScopeId FlatUsed = NestedUsed.Outer.AsFlat();
+			CHECK(Names[DoubleNested.Inner.Name] == "DoubleNested");
+			CHECK(Names[NestedUsed.Inner.Name] == "NestedUsed");
+			CHECK(Names[FlatUsed.Name] == "FlatUsed");
+
+			FParametricTypeView S4 = ResolveUntranslatedParametricType(Batch, S4T.Name.AsParametric());
+			CHECK(Names[S4.Name.Id] == "S4");
+			CHECK(S4.NumParameters == 2);
+			
+			FTypeId S2T = S4.Parameters[0];
+			FTypeId E2T = S4.Parameters[1];
+			CHECK(S2T.Scope == DoubleNested.Outer);
+			CHECK(E2T.Scope == DoubleNested.Outer);
+			CHECK(Names[S2T.Name.AsConcrete().Id] == "S2");
+			CHECK(Names[E2T.Name.AsConcrete().Id] == "E2");
+		});
+	}
+}
+
+//
+//namespace Test
+//{
+//
+//class FBindings
+//{
+//
+//	
+//};
+//	
+//struct FPt
+//{
+//	float x, y;
+//};
+//UE_REFLECT_STRUCT(FPt, A, B);
+//
+//
+//enum EFlat { A = 1, B = 3 };
+//enum EFlag { A = 0, B = 2 };
+//
+//UE_REFLECT_ENUM(EFlat, A, B);
+//UE_REFLECT_ENUM(EFlag, A, B);
+//
+//struct FEnums
+//{
+//	EFlat Flat;
+//	EFlag Flag;
+//};
+//
+//UE_REFLECT_STRUCT(FEnums, Flat, Flag);
+//
+//struct Leaves
+//{
+//	uint16 U16;
+//	int64 I64;
+//	bool Bool;
+//	char AsciiChar;
+//	uint8 BitfieldBool : 1;
+//	uint32 StaticU32s[3];
+//};
+//
+//UE_REFLECT_STRUCT(Leaves, U16, I64, Bool, AsciiChar, BitfieldBool, StaticU32s);
+//
+//GTestBindings.Add(LeavesCtti);
+//
+//struct FNestedStructs
+//{
+//	FPt Point;
+//	FPt StaticPoints[2];
+//};
+//UE_REFLECT_STRUCT(FNestedStructs, Point, StaticPoints);
+//
+//struct FLeafRanges
+//{
+//	TArray<int32> IntArray;
+//	TOptional<uint8> MaybeByte;
+//	TUniquePtr<float> FloatPtr;
+//	TSet<uint16> ShortSet;
+//	TSparseArray<bool> SparseBools;
+//};
+//
+//struct FSuper
+//{
+//	uint16 Pad;
+//	bool A;
+//};
+//UE_REFLECT_STRUCT(FSuper, A);
+//
+//struct FSub : FSuper
+//{
+//	bool B;
+//	uint32 Pad;
+//};
+//UE_REFLECT_SUBSTRUCT(FSub, FSuper, B);
+//
+//
+//struct FSubs
+//{
+//	TArray<FSub> Subs;
+//};
+
+
+TEST_CASE_NAMED(FPlainPropsLoadSaveTest, "System::Core::Serialization::PlainProps::LoadSave", "[Core][PlainProps][SmokeFilter]")
+{
+	SECTION("Leaves")
+	{}
+
+	SECTION("Enums")
+	{}
+		
+	SECTION("NestedStruct")
+	{}
+	
+	SECTION("StaticArray")
+	{}
+	
+	SECTION("LeafVariant")
+	{}
+	
+	SECTION("BitfieldBool")
+	{}
+
+	SECTION("LeafArray")
+	{}
+	
+	SECTION("LeafOptional")
+	{}
+	
+	SECTION("LeafSmartPtr")
+	{}
+
+	SECTION("LeafSetWhole")
+	{}
+
+	SECTION("LeafSparseArrayAppends")
+	{}
+
+	SECTION("LeafSetOps")
+	{}
+	
+	SECTION("SparseStructArray")
+	{}
+
+	SECTION("DenseStructArray")
+	{}
+	
+	SECTION("SubStructArray")
+	{}
+
+	SECTION("NestedLeafArray")
+	{}
+
+	SECTION("NestedStructArray")
+	{}
+
+	SECTION("StructToSubStructMapOps")
+	{}
+}
+
+
+//} // namespace Test
+
+
+} // namespace PlainProps
+//#endif // WITH_TESTS
\ No newline at end of file
diff --git a/Engine/Source/Programs/PlainPropsTests/Tests/PlainPropsUeCoreTests.cpp b/Engine/Source/Programs/PlainPropsTests/Tests/PlainPropsUeCoreTests.cpp
new file mode 100644
--- /dev/null
+++ b/Engine/Source/Programs/PlainPropsTests/Tests/PlainPropsUeCoreTests.cpp
@@ -0,0 +1,479 @@
+// Copyright Epic Games, Inc. All Rights Reserved.
+
+#if WITH_TESTS
+
+#include "Tests/TestHarnessAdapter.h"
+#include "Containers/StringView.h"
+#include  "PlainPropsBuildSchema.h"
+#include  "PlainPropsCtti.h"
+#include  "PlainPropsIndex.h"
+#include  "PlainPropsInternalBuild.h"
+#include  "PlainPropsInternalFormat.h"
+#include  "PlainPropsLoad.h"
+#include  "PlainPropsRead.h"
+#include  "PlainPropsSave.h"
+#include  "PlainPropsWrite.h"
+#include  "PlainPropsUeCoreBindings.h"
+#include "Templates/UnrealTemplate.h"
+
+namespace PlainProps::UE::Test
+{
+
+static TIdIndexer<FName>	GNames;
+static FDeclarations		GTypes;
+static FStructBindings		GBindings;
+static FStructBindings		GDeltaBindings;
+
+struct FIds
+{
+	static FMemberId		IndexMember(FAnsiStringView Name)			{ return GNames.NameMember(FName(Name)); }
+	static FTypenameId		IndexTypename(FAnsiStringView Name)			{ return GNames.MakeTypename(FName(Name)); }
+	static FScopeId			IndexNativeScope()							{ return GNames.MakeScope(FName(UE_MODULE_NAME)); }
+	static FTypeId			IndexNativeType(FAnsiStringView Typename)	{ return {IndexNativeScope(), IndexTypename(Typename)}; }
+	static FEnumSchemaId	IndexEnum(FTypeId Type)						{ return GNames.IndexEnum(Type); }
+	static FEnumSchemaId	IndexEnum(FAnsiStringView Name)				{ return IndexEnum(IndexNativeType(Name)); }
+	static FStructSchemaId	IndexStruct(FTypeId Type)					{ return GNames.IndexStruct(Type); }
+	static FStructSchemaId	IndexStruct(FAnsiStringView Name)			{ return IndexStruct(IndexNativeType(Name)); }
+	static const FDebugIds& GetDebug()									{ return GNames; }
+};
+
+////////////////////////////////////////////////////////////////////////////////////////////////
+
+template<class T>
+struct TCustomBindings
+{
+	using Type = void;
+};
+
+template<class T>
+struct TCustomDeltaBindings : TCustomBindings<T>
+{};
+
+template<typename T>
+struct TCustomDeltaBindings<TSet<T>>
+{
+	using Type = PlainProps::UE::TSetDeltaBinding<FIds, T>;
+};
+
+////////////////////////////////////////////////////////////////////////////////////////////////
+
+struct FDefaultRuntime
+{
+	using Ids = FIds;
+	template<class T> using CustomBindings = TCustomBindings<T>;
+
+	static FDeclarations&			GetDeclarations()				{ return GTypes; }
+	static FStructBindings&			GetBindings()					{ return GBindings; }
+	static void						DropStruct(FStructSchemaId Id)	{} // todo
+};
+
+struct FDeltaRuntime : FDefaultRuntime
+{
+	template<class T> using CustomBindings = TCustomDeltaBindings<T>;
+
+	static FStructBindings&			GetBindings()					{ return GDeltaBindings; }
+	static void						DropStruct(FStructSchemaId Id)	{} // todo,  drop from default too
+};
+
+////////////////////////////////////////////////////////////////////////////////////////////////
+
+template<class T, class Runtime>
+struct TScopedStructBinding
+{
+	using Ids = typename Runtime::Ids;
+	using Ctti = CttiOf<T>;
+	//constexpr bool bDeltaBindMembers = NumCustomMembers<Ctti, FDeltaRuntime::CustomBindings>() > NumCustomMembers<Ctti, FDefaultRuntime::CustomBindings>();
+
+	FStructSchemaId Id;
+
+	TScopedStructBinding(EMemberPresence Occupancy = EMemberPresence::AllowSparse)
+	: Id(DeclareNativeStruct<Ctti, Ids>(Runtime::GetDeclarations(), Occupancy))
+	{
+		//BindNativeStruct<Ctti, Runtime>(Runtime::GetBindings());
+		//if constexpr (bDeltaBindMembers)
+		//{
+		//	BindNativeStruct<Ctti, FDeltaRuntime>(Runtime::GetBindings());	
+		//}
+	}
+
+	~TScopedStructBinding()
+	{
+		Runtime::DropStruct(Id);
+	}
+};
+
+template<typename T, class Runtime>
+struct TScopedEnumBinding
+{
+	// todo
+};
+
+//template<typename T, class Runtime>
+//struct TScopedBinding : std::conditional_t<std::is_enum_v<T>, TScopedEnumBinding<T, Runtime>, TScopedStructBinding<T, Runtime>>
+//{
+//	
+//};
+
+//template<class T>
+//struct TScopedStructDeclaration
+//{
+//	FStructSchemaId Id;
+//
+//	TScopedStructDeclaration(EMemberPresence Occupancy)
+//	: Id(DeclareNativeStruct<CttiOf<T>>(FDefaultRuntime::GetDeclarations()), Occupancy)
+//	{}
+//
+//	~TScopedStructDeclaration()
+//	{
+//		FDeltaRuntime::DropStruct(Id);
+//	}
+//};
+////////////////////////////////////////////////////////////////////////////////////////////////
+
+inline constexpr uint32 Magics[] = { 0xFEEDF00D, 0xABCD1234, 0xDADADAAA, 0x99887766, 0xF0F1F2F3 };
+
+class FBatchSaver
+{
+public:
+	explicit FBatchSaver(const FStructBindings& InBindings) : Bindings(InBindings) {}
+
+	template<class T>
+	void						Save(T&& Object) { Save(IndexNativeStruct<T, FIds>(), &Object); }
+	TArray64<uint8>				Write() const;
+
+private:
+	void						Save(FStructSchemaId Id, const void* Object); 
+
+	using IdBuiltStructPair = TPair<FStructSchemaId, TUniquePtr<FBuiltStruct>>;
+	TArray<IdBuiltStructPair>	SavedObjects;
+	const FStructBindings&		Bindings;
+};
+
+void FBatchSaver::Save(FStructSchemaId Id, const void* Object)
+{
+	SavedObjects.Emplace(Id, SaveStruct(reinterpret_cast<const uint8*>(Object), Id, {GTypes, GBindings, /* debug */ GNames}));
+}
+
+template<typename ArrayType>
+void WriteNumAndArray(TArray64<uint8>& Out, const ArrayType& Items)
+{
+	WriteU32(Out, IntCastChecked<uint32>(Items.Num()));
+	WriteArray(Out, Items);
+}
+
+template<typename T>
+TConstArrayView<T> GrabNumAndArray(/* in-out */ FByteReader& It)
+{
+	uint32 Num = It.Grab<uint32>();
+	return MakeArrayView(reinterpret_cast<const T*>(It.GrabBytes(Num * sizeof(T))), Num);
+}
+
+TArray64<uint8> FBatchSaver::Write() const
+{
+	// Build partial schemas
+	FSchemasBuilder SchemaBuilders(GTypes.GetStructs(), GTypes.GetEnums(), /* debug */ GNames);
+	for (const IdBuiltStructPair& Object : SavedObjects)
+	{
+		SchemaBuilders.NoteMembers(Object.Key, *Object.Value);
+	}
+	FBuiltSchemas Schemas = SchemaBuilders.Build(); 
+
+	// Filter out declared but unused names and ids
+	FWriter Writer(GNames, Schemas);
+	TArray<FName> UsedNames;
+	for (uint32 Idx = 0, Num = GNames.NumNames(); Idx < Num; ++Idx)
+	{
+		if (Writer.Uses({Idx}))
+		{
+			UsedNames.Add(GNames.ResolveName({Idx}));
+		}
+	}
+	
+	// Write ids. Just copying in-memory FNames, a stable format might use SaveNameBatch().
+	TArray64<uint8> Out;
+	WriteU32(Out, Magics[0]);
+	WriteNumAndArray(Out, TArrayView<const FName, int32>(UsedNames));
+	//WriteNumAndArray(Out, Writer.GetUsedScopes());
+	//WriteNumAndArray(Out, Writer.GetUsedParametrics());
+
+	// Write schemas
+	WriteU32(Out, Magics[1]);
+	WriteAlignmentPadding<uint32>(Out);
+	TArray64<uint8> Tmp;
+	Writer.WriteSchemas(/* Out */ Tmp);
+	WriteNumAndArray(Out, TArrayView<const uint8, int64>(Tmp));
+	Tmp.Reset();
+
+	//// Write objects
+	WriteU32(Out, Magics[2]);
+	for (const TPair<FStructSchemaId, TUniquePtr<FBuiltStruct>>& Object : SavedObjects)
+	{
+		WriteU32(/* out */ Tmp, Magics[3]);
+		WriteU32(/* out */ Tmp, Writer.GetWriteId(Object.Key).Get().Idx);
+		Writer.WriteMembers(/* out */ Tmp, Object.Key, *Object.Value);
+		WriteSkippableSlice(Out, Tmp);
+		Tmp.Reset();
+	}
+
+	// Write object terminator
+	WriteSkippableSlice(Out, TConstArrayView64<uint8>());
+	WriteU32(Out, Magics[4]);
+		
+	return Out;
+}
+
+////////////////////////////////////////////////////////////////////////////////////////////////
+
+class FBatchLoader
+{
+public:
+	FBatchLoader(const FStructBindings& Bindings, FMemoryView Data)
+	{
+		// Read ids
+		FByteReader It(Data);
+		CHECK(It.Grab<uint32>() == Magics[0]);
+		Names = GrabNumAndArray<FName>(It);
+		
+		// Read schemas
+		CHECK(It.Grab<uint32>() == Magics[1]);
+		It.SkipAlignmentPadding<uint32>();
+		uint32 SchemasSize = It.Grab<uint32>();
+		const FSchemaBatch* SaveSchemas = ValidateSchemas(It.GrabSlice(SchemasSize));
+		CHECK(It.Grab<uint32>() == Magics[2]);
+		
+		// Bind saved ids to current ids, remap schemas with new ids and mount it
+		FIdBinder SaveToLoadIds(GNames, Names, SaveSchemas->GetNestedScopes(), SaveSchemas->GetParametricTypes());
+		FSchemaBatch* LoadSchemas = CreateTranslatedSchemas(SaveSchemas, SaveToLoadIds.Get());
+		FReadBatchId Batch = MountReadSchemas(LoadSchemas);
+
+		// Read objects
+		while (uint64 NumBytes = It.GrabVarIntU())
+		{	
+			FByteReader ObjIt(It.GrabSlice(NumBytes));
+			CHECK(ObjIt.Grab<uint32>() == Magics[3]);
+			FStructSchemaId Schema = { ObjIt.Grab<uint32>() };
+			Objects.Add({ { Schema, Batch }, ObjIt });
+		}
+		
+		CHECK(It.Grab<uint32>() == Magics[4]);
+		CHECK(!Objects.IsEmpty());
+
+		// Finally create load plans
+		Plans = CreateLoadPlans(Batch, GTypes, Bindings);
+	}
+
+	~FBatchLoader()
+	{
+		DestroyLoadPlans(Plans);
+		const FSchemaBatch* LoadSchemas = UnmountReadSchemas(Objects[0].Schema.Batch);
+		DestroyTranslatedSchemas(LoadSchemas);
+	}
+
+	template<class T>
+	T Load();
+
+private:
+	TConstArrayView<FName> Names;
+	FLoadBatch* Plans;
+	TArray<FStructView> Objects;
+	int32 LoadIdx = 0;
+
+	//virtual FNestedScope				Resolve(FNestedScopeId Id) const override						{ return Scopes[Id.Idx]; }
+	//virtual FParametricTypeView			Resolve(FParametricTypeId Id) const override					{ return Parametrics[Id.Idx]; }
+	//virtual void						AppendDebugString(FString& Out, FNameId Name) const override	{ return Names[Name.Idx].AppendString(Out); }
+};
+
+
+static void SaveAndLoad(const FStructBindings& Bindings, void (*Save)(FBatchSaver&), void (*Load)(FBatchLoader&))
+{
+	TArray64<uint8> Data;
+	{
+		FBatchSaver Batch(Bindings);
+		Save(Batch);
+		Data = Batch.Write();
+	}
+
+	FBatchLoader Batch(Bindings, MakeMemoryView(Data));
+	Load(Batch);
+}
+
+template<class Runtime>
+static void TestSaveAndLoad(void (*Save)(FBatchSaver&), void (*Load)(FBatchLoader&))
+{
+	SaveAndLoad(Runtime::GetBindings(), Save, Load);
+}
+
+
+
+//// Tests that everything was visited
+//struct FTestMemberVisitor : public FMemberVisitor
+//{
+//	using FMemberVisitor::FMemberVisitor;
+//	
+//	~FMemberVisitor()
+//	{
+//		CHECK(MemberIdx == NumMembers); // Must read all members
+//		CHECK(RangeTypeIdx == NumRangeTypes); // Must read all ranges
+//#if DO_CHECK
+//		CHECK(InnerSchemaIdx == NumInnerSchemas); // Must read all schema ids
+//#endif
+//	}
+//};
+//
+//template<typename OutType, typename InType>
+//TArray<OutType> MakeArray(const InType& Items)
+//{
+//	TArray<OutType> Out;
+//	Out.Reserve(IntCastChecked<int32>(Items.Num()));
+//	for (const auto& Item : Items)
+//	{
+//		Out.Emplace(Item);
+//	}
+//	return Out;
+//}
+//
+//enum EFlat { A = 1, B = 3 };
+//enum EFlag { A = 0, B = 2 };
+//
+//UE_REFLECT_ENUM(EFlat, A, B);
+//UE_REFLECT_ENUM(EFlag, A, B);
+//
+//struct FEnums
+//{
+//	EFlat Flat;
+//	EFlag Flag;
+//};
+//
+//UE_REFLECT_STRUCT(FEnums, Flat, Flag);
+//
+//struct Leaves
+//{
+//	uint16 U16;
+//	int64 I64;
+//	bool Bool;
+//	char AsciiChar;
+//	uint8 BitfieldBool : 1;
+//	uint32 StaticU32s[3];
+//};
+//
+//UE_REFLECT_STRUCT(Leaves, U16, I64, Bool, AsciiChar, BitfieldBool, StaticU32s);
+//
+//GTestBindings.Add(LeavesCtti);
+//
+//struct FNestedStructs
+//{
+//	FPt Point;
+//	FPt StaticPoints[2];
+//};
+//UE_REFLECT_STRUCT(FNestedStructs, Point, StaticPoints);
+//
+//struct FLeafRanges
+//{
+//	TArray<int32> IntArray;
+//	TOptional<uint8> MaybeByte;
+//	TUniquePtr<float> FloatPtr;
+//	TSet<uint16> ShortSet;
+//	TSparseArray<bool> SparseBools;
+//};
+//
+//struct FSuper
+//{
+//	uint16 Pad;
+//	bool A;
+//};
+//UE_REFLECT_STRUCT(FSuper, A);
+//
+//struct FSub : FSuper
+//{
+//	bool B;
+//	uint32 Pad;
+//};
+//UE_REFLECT_SUBSTRUCT(FSub, FSuper, B);
+//
+//
+//struct FSubs
+//{
+//	TArray<FSub> Subs;
+//};
+
+
+struct FInt { int32 X; };
+PP_REFLECT_STRUCT(PlainProps::UE::Test, FInt, void, X);
+
+TEST_CASE_NAMED(FPlainPropsUeCoreTest, "System::Core::Serialization::PlainProps::UE::Core", "[Core][PlainProps][SmokeFilter]")
+{
+	SECTION("Basic")
+	{
+		TScopedStructBinding<FInt, FDefaultRuntime> Int;
+
+		//TestSaveAndLoad<FDefaultRuntime>(
+		//[](FBatchSaver& Batch)
+		//{
+		//	//Batch.Save<FDefaultRuntime>(FInt{1234});
+		//}, 
+		//[](FBatchLoader& Batch)
+		//{
+		//	//CHECK(Batch.Load<FInt>().X == 1234);
+		//});
+	}
+
+	SECTION("TUniquePtr")
+	{}
+
+	SECTION("TArray")
+	{}
+
+	SECTION("FString")
+	{}
+		
+	SECTION("TSet")
+	{}
+	
+	SECTION("FName")
+	{}
+
+	SECTION("FakeReference")
+	{}
+
+	SECTION("NestedContainer")
+	{}
+	
+	SECTION("TSetDelta")
+	{}
+
+	//SECTION("LeafOptional")
+	//{}
+	//
+	//SECTION("LeafSmartPtr")
+	//{}
+
+	//SECTION("LeafSetWhole")
+	//{}
+
+	//SECTION("LeafSparseArrayAppends")
+	//{}
+
+	//SECTION("LeafSetOps")
+	//{}
+	//
+	//SECTION("SparseStructArray")
+	//{}
+
+	//SECTION("DenseStructArray")
+	//{}
+	//
+	//SECTION("SubStructArray")
+	//{}
+
+	//SECTION("NestedLeafArray")
+	//{}
+
+	//SECTION("NestedStructArray")
+	//{}
+
+	//SECTION("StructToSubStructMapOps")
+	//{}
+}
+
+} // namespace PlainProps::UE::Test
+#endif // WITH_TESTS
\ No newline at end of file
