PlainProps flesh out bind, load and ue core tests a bit more

diff --git a/Engine/Source/Programs/PlainPropsTests/Private/PlainPropsBind.cpp b/Engine/Source/Programs/PlainPropsTests/Private/PlainPropsBind.cpp
--- a/Engine/Source/Programs/PlainPropsTests/Private/PlainPropsBind.cpp
+++ b/Engine/Source/Programs/PlainPropsTests/Private/PlainPropsBind.cpp
@@ -1,9 +1,9 @@
 // Copyright Epic Games, Inc. All Rights Reserved.
 
 #include  "PlainPropsBind.h"
-//#include  "PlainPropsInternalFormat.h"
-//#include  "PlainPropsInternalRead.h"
-//#include <type_traits>
+#include  "PlainPropsIndex.h"
+#include  "PlainPropsInternalFormat.h"
+#include  "PlainPropsInternalRead.h"
 
 namespace PlainProps
 {
@@ -17,6 +17,19 @@ static_assert((uint8)ELeafType::Hex			== (uint8)ELeafBindType::Hex);
 static_assert((uint8)ELeafType::Enum		== (uint8)ELeafBindType::Enum);
 static_assert((uint8)ELeafType::Unicode		== (uint8)ELeafBindType::Unicode);
 
+////////////////////////////////////////////////////////////////////////////////////////////////
+
+uint32 FStructSchemaBinding::CalculateSize() const
+{
+	uint32 Out = sizeof(FStructSchemaBinding) + NumMembers * sizeof(FMemberBindType);
+	Out = Align(Out + NumMembers * sizeof(uint32), sizeof(uint32));
+	Out = Align(Out + NumInnerSchemas * sizeof(FSchemaId), sizeof(FSchemaId));
+	Out = Align(Out + NumInnerRanges * sizeof(FRangeBinding), sizeof(FRangeBinding));
+	return Out;
+}
+
+////////////////////////////////////////////////////////////////////////////////////////////////
+
 FMemberVisitor::FMemberVisitor(const FStructSchemaBinding& InSchema)
 : Schema(InSchema)
 , NumMembers(InSchema.NumMembers)
@@ -50,12 +63,11 @@ FLeafMemberBinding FMemberVisitor::GrabLeaf()
 FStructMemberBinding FMemberVisitor::GrabStruct()
 {
 	checkf(!PeekType().AsStruct().IsDynamic, TEXT("Bound structs can't be dynamic"));
-	return { static_cast<FStructSchemaId>(GrabInnerSchema()), GrabMemberOffset() };
+	return { PeekType().AsStruct(), static_cast<FStructSchemaId>(GrabInnerSchema()), GrabMemberOffset() };
 }
 
 static bool HasSchema(FMemberBindType Type)
 {
-	checkf(!Type.IsStruct() || !Type.AsStruct().IsDynamic, TEXT("Bound structs can't be dynamic"));
 	return Type.IsStruct() || Type.AsLeaf().Bind.Type == ELeafBindType::Enum;
 }
 
@@ -76,7 +88,23 @@ FRangeMemberBinding FMemberVisitor::GrabRange()
 	FOptionalSchemaId InnermostSchema = HasSchema(InnerTypes.Last()) ? ToOptional(GrabInnerSchema()) : NoId; 
 	uint64 Offset = GrabMemberOffset();
 		
-	return { &InnerTypes[0], RangeBindings, InnermostSchema, Offset};
+	return { &InnerTypes[0], RangeBindings, static_cast<uint32>(InnerTypes.Num()), InnermostSchema, Offset};
+}
+
+void FMemberVisitor::SkipMember()
+{
+	FMemberBindType Type = PeekType();
+	if (Type.IsRange())
+	{
+		FMemberBindTypeRange InnerTypes = GrabInnerTypes();
+		InnerSchemaIdx += HasSchema(InnerTypes.Last());
+	}
+	else
+	{
+		InnerSchemaIdx += HasSchema(Type);
+	}
+	
+	++MemberIdx;
 }
 
 FSchemaId FMemberVisitor::GrabInnerSchema()
@@ -92,7 +120,6 @@ FRangeBinding::FRangeBinding(const IRangeBinding& Binding, ERangeSizeType SizeTy
 	check(SizeType == GetSizeType());
 }
 
-
 ////////////////////////////////////////////////////////////////////////////////////////////////
 
 FStructBindings::~FStructBindings()
@@ -133,62 +160,28 @@ static uint16 CountRanges(TConstArrayView<FMemberBinding> Members)
 
 void FStructBindings::BindStruct(FStructSchemaId Id, TConstArrayView<FMemberBinding> Members)
 {
+	// Make header, allocate and copy header
 	FStructSchemaBinding Header = { IntCastChecked<uint16>(Members.Num()), CountInnerSchemas(Members), CountRanges(Members) };
+	FStructSchemaBinding* Schema = new (FMemory::MallocZeroed(Header.CalculateSize())) FStructSchemaBinding {Header};
 
-	// Calculate size
-	uint64 Size = sizeof(Header);
-	Size += (Members.Num() + Header.NumInnerRanges) * sizeof(FMemberBindType);
-	Size = Align(Size, sizeof(uint32));
-	Size += Members.Num() * sizeof(uint32);
-	Size = Align(Size, sizeof(FSchemaId));
-	Size += Header.NumInnerSchemas * sizeof(FSchemaId);
-	Size = Align(Size, sizeof(FRangeBinding));
-	Size += Header.NumInnerRanges * sizeof(FRangeBinding);
-
-	// Allocate and copy header
-	FStructSchemaBinding* Schema = new (FMemory::MallocZeroed(Size)) FStructSchemaBinding {Header};
-
-	// Copy footer
-	FMemberBindType* MemberIt		= Schema->Members;
-	FMemberBindType* RangeTypeIt	= const_cast<FMemberBindType*>(Schema->GetInnerRangeTypes());
-	uint32* OffsetIt				= const_cast<uint32*>(Schema->GetOffsets());
-	FSchemaId* InnerSchemaIt		= const_cast<FSchemaId*>(Schema->GetInnerSchemas());
-	FRangeBinding* RangeBindingIt	= const_cast<FRangeBinding*>(Schema->GetRangeBindings());
+	// Write footer
+	FMemberBinder Footer(*Schema);
 	for (const FMemberBinding& Member : Members)
 	{
-		if (TConstArrayView<FRangeBinding> Ranges = Member.RangeBindings; Ranges.IsEmpty())
+		TConstArrayView<FRangeBinding> Ranges = Member.RangeBindings;
+		if (Ranges.IsEmpty())
 		{
-			*MemberIt++ = Member.InnermostType;
+			Footer.AddMember(Member.InnermostType, IntCastChecked<uint32>(Member.Offset));
 		}
 		else
 		{
-			*MemberIt++ = FMemberBindType(Ranges[0].GetSizeType());
-
-			for (FRangeBinding Range : Ranges.RightChop(1))
-			{
-				*RangeTypeIt++ = FMemberBindType(Range.GetSizeType());
-			}
-			*RangeTypeIt++ = Member.InnermostType;
-
-			FMemory::Memcpy(RangeBindingIt, Ranges.GetData(), Ranges.Num() * Ranges.GetTypeSize());
-			RangeBindingIt += Ranges.Num();
-		}
-
-		if (Member.InnermostSchema)
-		{
-			*InnerSchemaIt++ = Member.InnermostSchema.Get();
+			Footer.AddRange(Ranges, Member.InnermostType, IntCastChecked<uint32>(Member.Offset));
 		}
 
-		*OffsetIt++ = IntCastChecked<uint32>(Member.Offset);
+		Footer.AddOptionalInnerSchema(Member.InnermostSchema);
 	}
 
-	// Validate copying
-	check(MemberIt == Schema->GetInnerRangeTypes());
-	check(RangeTypeIt == (const void*)Schema->GetOffsets());
-	check(OffsetIt == (const void*)Schema->GetInnerSchemas());
-	check(InnerSchemaIt == (const void*)Schema->GetRangeBindings());
-	check(Header.NumInnerRanges == RangeBindingIt - Schema->GetRangeBindings());
-
+	// Register
 	Bind(Id, FStructBinding(*Schema));
 }
 
@@ -212,4 +205,118 @@ FStructBinding FStructBindings::Get(FStructSchemaId Id) const
 
 //void FStructBindings::DropStruct(FStructSchemaId Id);
 
+
+//////////////////////////////////////////////////////////////////////////
+
+uint32 FIdTranslatorBase::CalculateTranslationSize(int32 NumSavedNames, const FSchemaBatch& Batch)
+{
+	static_assert(sizeof(FNameId) == sizeof(FNestedScopeId));
+	static_assert(sizeof(FNameId) == sizeof(FParametricTypeId));
+	static_assert(sizeof(FNameId) == sizeof(FSchemaId));
+	return sizeof(FNameId) * (NumSavedNames + Batch.NumNestedScopes + Batch.NumParametricTypes + Batch.NumSchemas);
+}
+
+FFlatScopeId Translate(FFlatScopeId From, TConstArrayView<FNameId> ToNames)
+{
+	return { ToNames[From.Name.Idx] };
+}
+
+static void TranslateScopeIds(TArrayView<FNestedScopeId> Out, FIdIndexerBase& Indexer, TConstArrayView<FNameId> ToNames, TConstArrayView<FNestedScope> From)
+{
+	uint32 OutIdx = 0;
+	for (FNestedScope Scope : From)
+	{
+		check(Scope.Outer.IsFlat() || Scope.Outer.AsNested().Idx < OutIdx);
+		FScopeId Outer = Scope.Outer.IsFlat() ? FScopeId(Translate(Scope.Outer.AsFlat(), ToNames)) : FScopeId(Out[Scope.Outer.AsNested().Idx]);
+		FFlatScopeId Inner = Translate(Scope.Inner, ToNames);
+		Out[OutIdx] = Indexer.NestScope(Outer, Inner).AsNested();
+	}
+}
+
+static void TranslateParametricTypeIds(TArrayView<FParametricTypeId> Out, FIdIndexerBase& Indexer, const FIdBinding& To, TConstArrayView<FParametricType> From, const FTypeId* FromParameters)
+{
+	TArray<FTypeId, TInlineAllocator<8>> Params;
+	uint32 OutIdx = 0;
+	for (FParametricType Parametric : From)
+	{
+		Params.Reset();
+		for (FTypeId FromParameter : MakeArrayView(FromParameters + Parametric.Parameters.Idx, Parametric.Parameters.NumParameters))
+		{
+			Params.Add(To.Remap(FromParameter));
+		}
+		Out[OutIdx] = Indexer.MakeParametricTypeId(To.Remap(Parametric.Name), Params);
+	}
+}
+
+static void TranslateSchemaIds(TArrayView<FSchemaId> Out, FIdIndexerBase& Indexer, const FIdBinding& To, const FSchemaBatch& From)
+{
+	uint32 OutIdx = 0;
+	for (const FStructSchema& FromSchema : GetStructSchemas(From))
+	{
+		FTypeId ToType = To.Remap(FromSchema.Type);
+		Out[OutIdx++] = Indexer.IndexStruct(ToType);
+	}
+	
+	for (const FEnumSchema& FromSchema : GetEnumSchemas(From))
+	{
+		FTypeId ToType = To.Remap(FromSchema.Type);
+		Out[OutIdx++] = Indexer.IndexEnum(ToType);
+	}
+}
+
+FIdBinding FIdTranslatorBase::TranslateIds(FMutableMemoryView To, FIdIndexerBase& Indexer, TConstArrayView<FNameId> ToNames, const FSchemaBatch& From)
+{
+	TArrayView<FNestedScopeId> ToScopes(static_cast<FNestedScopeId*>(To.GetData()), From.NumNestedScopes);
+	TArrayView<FParametricTypeId> ToParametricTypes(reinterpret_cast<FParametricTypeId*>(ToScopes.end()), From.NumParametricTypes);
+	TArrayView<FSchemaId> ToSchemas(reinterpret_cast<FSchemaId*>(ToParametricTypes.end()), From.NumSchemas);
+	FIdBinding Out = {ToNames, ToScopes, ToParametricTypes, ToSchemas};
+	check(uintptr_t(To.GetDataEnd()) == uintptr_t(ToSchemas.end()));
+
+	TranslateScopeIds(ToScopes, Indexer, ToNames, From.GetNestedScopes());
+	TranslateParametricTypeIds(ToParametricTypes, Indexer, Out, From.GetParametricTypes(), From.GetFirstParameter());
+	TranslateSchemaIds(ToSchemas, Indexer, Out, From);
+
+	return Out;
+}
+
+//////////////////////////////////////////////////////////////////////////
+
+FSchemaBatch* CreateTranslatedSchemas(const FSchemaBatch& In, FIdBinding NewIds)
+{
+	const FMemoryView InSchemas = GetSchemaData(In);
+	const uint32 Num = In.NumSchemas;
+	const uint64 Size = sizeof(FSchemaBatch) + /* offsets */ sizeof(uint32) * Num + InSchemas.GetSize();
+
+	// Allocate and copy header
+	FSchemaBatch* Out = new (FMemory::Malloc(Size)) FSchemaBatch {In};
+	Out->NumNestedScopes = 0;
+	Out->NestedScopesOffset = 0;
+	Out->NumParametricTypes = 0;
+
+	// Initialize schema offsets
+	const uint32 DroppedBytes = IntCastChecked<uint32>(uintptr_t(InSchemas.GetData()) - uintptr_t(In.SchemaOffsets + Num));
+	for (uint32 Idx = 0; Idx < Num; ++Idx)
+	{
+		Out->SchemaOffsets[Idx] = In.SchemaOffsets[Idx] - DroppedBytes;
+	}
+
+	// Copy schemas and remap type ids
+	FMemory::Memcpy(reinterpret_cast<uint8*>(Out) + Out->GetSchemaOffsets()[0], InSchemas.GetData(), InSchemas.GetSize());
+	for (FStructSchema& Schema : GetStructSchemas(*Out))
+	{
+		Schema.Type = NewIds.Remap(Schema.Type);
+	}
+	for (FEnumSchema& Schema : GetEnumSchemas(*Out))
+	{
+		Schema.Type = NewIds.Remap(Schema.Type);
+	}
+
+	return Out;
+}
+
+void DestroyTranslatedSchemas(const FSchemaBatch* Schemas)
+{
+	FMemory::Free(const_cast<FSchemaBatch*>(Schemas));
+}
+
 } // namespace PlainProps
\ No newline at end of file
diff --git a/Engine/Source/Programs/PlainPropsTests/Private/PlainPropsIndex.cpp b/Engine/Source/Programs/PlainPropsTests/Private/PlainPropsIndex.cpp
--- a/Engine/Source/Programs/PlainPropsTests/Private/PlainPropsIndex.cpp
+++ b/Engine/Source/Programs/PlainPropsTests/Private/PlainPropsIndex.cpp
@@ -127,10 +127,14 @@ FScopeId FIdIndexerBase::NestScope(FScopeId Outer, FFlatScopeId Inner)
 	return FScopeId(NestedScopes.Index(Outer, Inner));
 }
 
+FParametricTypeId FIdIndexerBase::MakeParametricTypeId(FConcreteTypenameId Name, TConstArrayView<FTypeId> Params)
+{
+	return ParametricTypes.Index({Name, IntCastChecked<uint8>(Params.Num()), Params.GetData()});
+}
+
 FTypeId FIdIndexerBase::MakeParametricType(FTypeId Type, TConstArrayView<FTypeId> Params)
 {
-	FParametricTypeId ParametricType = ParametricTypes.Index({Type.Name.AsConcrete(), IntCastChecked<uint8>(Params.Num()), Params.GetData()});
-	return {Type.Scope, FTypenameId(ParametricType)};
+	return {Type.Scope, FTypenameId(MakeParametricTypeId(Type.Name.AsConcrete(), Params))};
 }
 
 FEnumSchemaId FIdIndexerBase::IndexEnum(FTypeId Type)
@@ -142,15 +146,5 @@ FStructSchemaId	FIdIndexerBase::IndexStruct(FTypeId Type)
 {
 	return { IntCastChecked<uint32>(Structs.Add(Type).AsInteger()) };
 }
-//
-//FParametricTypeView	FIdIndexerBase::Resolve(FParametricTypeId Id) const
-//{
-//	return ParametricTypes.Resolve(Id);
-//}
-//
-//FNestedScope FIdIndexerBase::Resolve(FNestedScopeId Id) const
-//{
-//	return NestedScopes.Resolve(Id);
-//}
 
 } // namespace PlainProps
\ No newline at end of file
diff --git a/Engine/Source/Programs/PlainPropsTests/Private/PlainPropsInternalFormat.h b/Engine/Source/Programs/PlainPropsTests/Private/PlainPropsInternalFormat.h
--- a/Engine/Source/Programs/PlainPropsTests/Private/PlainPropsInternalFormat.h
+++ b/Engine/Source/Programs/PlainPropsTests/Private/PlainPropsInternalFormat.h
@@ -7,37 +7,46 @@
 namespace PlainProps
 {
 
+struct FSchemaBatchIterator;
+
 struct FSchemaBatch
 {
 	uint32 NumNestedScopes;
 	uint32 NestedScopesOffset;
 	uint32 NumParametricTypes;
 	uint32 NumSchemas;
+	uint32 NumStructSchemas;
 	uint32 SchemaOffsets[0];
 
+	TConstArrayView<uint32> GetSchemaOffsets() const
+	{
+		return MakeArrayView(SchemaOffsets, NumSchemas);
+	}
+
 	TConstArrayView<FNestedScope> GetNestedScopes() const
 	{
-		return MakeArrayView(reinterpret_cast<const FNestedScope*>(reinterpret_cast<const uint8*>(this) + NestedScopesOffset), NumNestedScopes);
+		return MakeArrayView(reinterpret_cast<const FNestedScope*>(uintptr_t(this) + NestedScopesOffset), NumNestedScopes);
 	}
 
 	TConstArrayView<FParametricType> GetParametricTypes() const
 	{
-		return MakeArrayView(reinterpret_cast<const FParametricType*>(GetNestedScopes().GetData() + NumNestedScopes), NumParametricTypes);
+		return MakeArrayView(reinterpret_cast<const FParametricType*>(GetNestedScopes().end()), NumParametricTypes);
 	}
 
 	const FTypeId* GetFirstParameter() const
 	{
-		return reinterpret_cast<const FTypeId*>(GetParametricTypes().GetData() + NumParametricTypes);
+		return reinterpret_cast<const FTypeId*>(GetParametricTypes().end());
 	}
 
 	void ValidateBounds(uint64 NumBytes) const;
 };
 
-static_assert(sizeof(FSchemaBatch) == 16 && alignof(FSchemaBatch) == 4,		"Add binary format versioning and read old format in AllocateReadSchemas");
+static_assert(sizeof(FSchemaBatch) == 20 && alignof(FSchemaBatch) == 4,			"Add binary format versioning and read old format in AllocateReadSchemas");
 static_assert(offsetof(FSchemaBatch, NumNestedScopes) == 0,						"Add binary format versioning and read old format in AllocateReadSchemas");
 static_assert(offsetof(FSchemaBatch, NestedScopesOffset) == 4,					"Add binary format versioning and read old format in AllocateReadSchemas");
 static_assert(offsetof(FSchemaBatch, NumParametricTypes) == 8,					"Add binary format versioning and read old format in AllocateReadSchemas");
 static_assert(offsetof(FSchemaBatch, NumSchemas) == 12,							"Add binary format versioning and read old format in AllocateReadSchemas");
+static_assert(offsetof(FSchemaBatch, NumStructSchemas) == 16,					"Add binary format versioning and read old format in AllocateReadSchemas");
 
 //////////////////////////////////////////////////////////////////////////
 
@@ -88,6 +97,8 @@ struct FStructSchema
 		const FSchemaId* FirstSchema = GetInnerSchemas();
 		return Inheritance != ESuper::No ? ToOptional(static_cast<FStructSchemaId>(*FirstSchema)) : NoId;
 	}
+
+	TConstArrayView<FMemberId> GetMemberNames() const { return MakeArrayView(GetMemberNames(Footer, NumMembers, NumRangeTypes), NumMembers); }
 };
 
 static_assert(sizeof(FStructSchema) == 16 && alignof(FStructSchema) == 4,		"Add binary format versioning and read old format in AllocateReadSchemas");
@@ -99,6 +110,15 @@ static_assert(offsetof(FStructSchema, NumRangeTypes) == 10,						"Add binary for
 static_assert(offsetof(FStructSchema, NumInnerSchemas) == 12,					"Add binary format versioning and read old format in AllocateReadSchemas");
 static_assert(offsetof(FStructSchema, Footer) == 15,							"Add binary format versioning and read old format in AllocateReadSchemas");
 
+inline uint32 CalculateSize(const FStructSchema& In)
+{
+	static_assert(alignof(FMemberType) == 1);
+	uint32 Out = sizeof(In) + sizeof(FMemberType) * (uint32(In.NumMembers) + In.NumRangeTypes);
+	Out = Align(Out + In.NumMembers * sizeof(FMemberId), sizeof(FMemberId));
+	Out = Align(Out + (In.NumInnerSchemas) * sizeof(FSchemaId), sizeof(FSchemaId));
+	return Out;
+}
+
 //////////////////////////////////////////////////////////////////////////
 
 struct FEnumSchema
@@ -117,11 +137,21 @@ static_assert(offsetof(FEnumSchema, Width) == 9,						"Add binary format version
 static_assert(offsetof(FEnumSchema, Num) == 10,							"Add binary format versioning and read old format in AllocateReadSchemas");
 static_assert(offsetof(FEnumSchema, Footer) == 12,						"Add binary format versioning and read old format in AllocateReadSchemas");
 
+inline uint32 CalculateSize(const FEnumSchema& In)
+{
+	return Align(sizeof(In) + In.Num * sizeof(In.Footer[0]) + In.ExplicitConstants * In.Num * SizeOf(In.Width), alignof(FEnumSchema));
+}
+
 //////////////////////////////////////////////////////////////////////////
 
+inline bool IsEnum(FMemberType Type)
+{
+	return Type.GetKind() == EMemberKind::Leaf && ELeafType::Enum == Type.AsLeaf().Type;
+}
+
 inline bool IsStructOrEnum(FMemberType Type)
 {
-	return Type.GetKind() == EMemberKind::Struct || (Type.GetKind() == EMemberKind::Leaf && ELeafType::Enum == Type.AsLeaf().Type);
+	return Type.IsStruct() || IsEnum(Type);
 }
 
 inline bool IsSuper(FMemberType Type)
@@ -134,6 +164,10 @@ inline constexpr uint64 GetLeafRangeSize(uint64 Num, FUnpackedLeafType Leaf)
 	return Leaf.Type == ELeafType::Bool ? (Num + 7) / 8 : Num * SizeOf(Leaf.Width);
 }
 
+
+
+
+
 //////////////////////////////////////////////////////////////////////////
 
 } // namespace PlainProps
\ No newline at end of file
diff --git a/Engine/Source/Programs/PlainPropsTests/Private/PlainPropsInternalRead.h b/Engine/Source/Programs/PlainPropsTests/Private/PlainPropsInternalRead.h
--- a/Engine/Source/Programs/PlainPropsTests/Private/PlainPropsInternalRead.h
+++ b/Engine/Source/Programs/PlainPropsTests/Private/PlainPropsInternalRead.h
@@ -4,6 +4,8 @@
 
 #include  "PlainPropsTypes.h"
 #include  "PlainPropsRead.h"
+#include  "PlainPropsInternalFormat.h"
+#include <type_traits>
 
 namespace PlainProps
 {
@@ -37,4 +39,57 @@ inline uint64 GrabRangeNum(ERangeSizeType MaxSize, FByteReader& ByteIt, FBitCach
 	return 0;
 }
 
-} // namespace PlainProps
\ No newline at end of file
+//////////////////////////////////////////////////////////////////////////
+
+template<class T>
+struct TSchemaIterator
+{
+	uintptr_t						Base;
+	const uint32*					OffsetIt;
+	
+	void							operator++()							{ ++OffsetIt; }
+	bool							operator!=(TSchemaIterator O) const		{ return OffsetIt != O.OffsetIt; }
+	T&								operator*() const						{ return *reinterpret_cast<T*>(Base + *OffsetIt); }
+};
+
+template<class T>
+struct TSchemaRange
+{
+	static constexpr bool bStructSchema = std::is_convertible_v<T, FStructSchema>;
+
+	TSchemaRange(std::conditional_t<std::is_const_v<T>, const FSchemaBatch&, FSchemaBatch&> Batch)
+	: Base(reinterpret_cast<uintptr_t>(&Batch))
+	, Offsets(Batch.SchemaOffsets + (bStructSchema ? 0 : Batch.NumStructSchemas), bStructSchema ? Batch.NumStructSchemas : Batch.NumSchemas - Batch.NumStructSchemas)
+	{
+		static_assert(bStructSchema ^ std::is_convertible_v<T, FEnumSchema>);
+	}
+	
+	TSchemaIterator<T>				begin() const		{ return { Base, Offsets.begin() }; }
+	TSchemaIterator<T>				end() const			{ return { Base, Offsets.end() }; }
+	T&								First() const		{ return *reinterpret_cast<T*>(Base + Offsets[0]); }
+	T&								Last() const		{ return *reinterpret_cast<T*>(Base + Offsets.Last()); }
+
+	uintptr_t						Base;
+	TConstArrayView<uint32>			Offsets;
+};
+
+inline TSchemaRange<FEnumSchema>			GetEnumSchemas(FSchemaBatch& Batch)				{ return { Batch }; }
+inline TSchemaRange<const FEnumSchema>		GetEnumSchemas(const FSchemaBatch& Batch)		{ return { Batch }; }
+inline TSchemaRange<FStructSchema>			GetStructSchemas(FSchemaBatch& Batch)			{ return { Batch }; }
+inline TSchemaRange<const FStructSchema>	GetStructSchemas(const FSchemaBatch& Batch)		{ return { Batch }; }
+
+//////////////////////////////////////////////////////////////////////////
+
+inline FMemoryView GetSchemaData(const FSchemaBatch& Batch)
+{
+	check(Batch.NumSchemas > 0);
+	const uint8* BatchPtr = reinterpret_cast<const uint8*>(&Batch);
+	TConstArrayView<uint32> Offsets(Batch.SchemaOffsets, Batch.NumSchemas);
+	
+	bool bHasEnums = Batch.NumSchemas > Batch.NumStructSchemas;
+	uint32 LastSchemaSize = bHasEnums	? CalculateSize(GetEnumSchemas(Batch).Last())
+										: CalculateSize(GetStructSchemas(Batch).Last());
+	return FMemoryView(BatchPtr + Offsets[0], Offsets.Last() + LastSchemaSize - Offsets[0]);
+}
+
+} // namespace PlainProps
diff --git a/Engine/Source/Programs/PlainPropsTests/Private/PlainPropsLoad.cpp b/Engine/Source/Programs/PlainPropsTests/Private/PlainPropsLoad.cpp
--- a/Engine/Source/Programs/PlainPropsTests/Private/PlainPropsLoad.cpp
+++ b/Engine/Source/Programs/PlainPropsTests/Private/PlainPropsLoad.cpp
@@ -4,6 +4,7 @@
 #include  "PlainPropsBind.h"
 #include  "PlainPropsInternalFormat.h"
 #include  "PlainPropsInternalRead.h"
+#include "Misc/Optional.h"
 #include <type_traits>
 
 namespace PlainProps
@@ -35,15 +36,16 @@ public:
 
 	// @param OffsetWidth Usage unimplemented, store size and offsets as 8/16/32/64-bit 
 	explicit FLoadStructPlan(const FStructSchemaBinding& Schema, ELeafWidth OffsetWidth, bool bSparse)
-	: Handle(uint64(&Schema) | (uint64(OffsetWidth) << 1) | (uint64(bSparse) << KernelAddressBit) | uint64(1))
+	: Handle(uint64(&Schema) | (uint64(OffsetWidth) << 1) | (uint64(bSparse) << KernelAddressBit) | SchemaMask)
 	{
 		check(&Schema == &AsSchema());
+		check(IsSparseSchema() == bSparse);
 	}
 
-	bool						IsSchema() const		{ return !!(Handle & SchemaBit); }
-	bool						IsSparseSchema() const	{ return !!(Handle & KernelAddressBit); }
-	bool						IsMemcpy() const		{ return (Handle & NonSchemaMask) == MemcpyMask; }
-	bool						IsCustom() const		{ return (Handle & NonSchemaMask) == CustomMask; }
+	bool						IsSchema() const		{ return (Handle & SchemaMask) == SchemaMask; }
+	bool						IsSparseSchema() const	{ return (Handle & SparseSchemaMask) == SparseSchemaMask; }
+	bool						IsMemcpy() const		{ return (Handle & LoMask) == MemcpyMask; }
+	bool						IsCustom() const		{ return (Handle & LoMask) == CustomMask; }
 	FLoadStructMemcpy			AsMemcpy() const		{ check(IsMemcpy()); return { static_cast<uint32>(Handle >> 32), static_cast<uint32>(Handle) >> 1 }; }
 	const ICustomStructBinding&	AsCustom() const		{ check(IsCustom()); return *AsPtr<ICustomStructBinding>(); }
 	const FStructSchemaBinding&	AsSchema() const		{ check(IsSchema()); return *AsPtr<FStructSchemaBinding>(); }
@@ -60,11 +62,14 @@ private:
 #endif
 	// todo handle WASM, copy updated KernelAddressBit from AssetDataTagMap.h
 
-	static constexpr uint64 PtrMask = (uint64(1) << KernelAddressBit) | 0b111;
-	static constexpr uint64 NonSchemaMask = 0b11;
-	static constexpr uint64 MemcpyMask = 0b00;
-	static constexpr uint64 CustomMask = 0b10;
-	static constexpr uint64 SchemaBit = 1;
+	
+	static constexpr uint64 SparseMask			= uint64(1) << KernelAddressBit;
+	static constexpr uint64 PtrMask				= ~(SparseMask | 0b111);
+	static constexpr uint64 LoMask				= 0b11;
+	static constexpr uint64 MemcpyMask			= 0b00;
+	static constexpr uint64 CustomMask			= 0b10;
+	static constexpr uint64 SchemaMask			= 0b01;
+	static constexpr uint64 SparseSchemaMask	= SchemaMask | SparseMask;
 	
 	template<typename T>
 	const T* AsPtr() const
@@ -78,16 +83,50 @@ private:
 
 ////////////////////////////////////////////////////////////////////////////
 
-static bool ContainsEnums(const FStructSchemaBinding& Schema)
+static uint16 CountEnums(const FStructSchema& Schema)
+{
+	uint16 Num = 0;
+	for (FMemberType Member : MakeArrayView(Schema.Footer, Schema.NumMembers))
+	{
+		Num += IsEnum(Member);
+	}
+	return Num;
+}
+
+static uint16 CountEnums(const FStructSchemaBinding& Schema)
 {
+	uint16 Num = 0;
 	for (FMemberBindType Member : MakeArrayView(Schema.Members, Schema.NumMembers))
 	{
-		if (Member.IsLeaf() && Member.AsLeaf().Bind.Type == ELeafBindType::Enum)
-		{
-			return true;
-		}
+		Num += (Member.IsLeaf() && Member.AsLeaf().Bind.Type == ELeafBindType::Enum);
+	}
+	return Num;
+}
+
+static uint16 CountStaticStructs(const FStructSchemaBinding& Schema)
+{
+	uint16 Num = 0;
+	for (FMemberBindType Member : MakeArrayView(Schema.Members, Schema.NumMembers))
+	{
+		Num += (Member.IsStruct() && !Member.AsStruct().IsDynamic);
+	}
+	return Num;
+}
+
+static bool HasDifferentSupers(const FStructSchema& From, const FStructSchemaBinding& To, TConstArrayView<FStructSchemaId> ToStructIds)
+{
+	if (From.Inheritance == ESuper::No)
+	{
+		return To.HasSuper();
 	}
-	return false;
+	else if (To.HasSuper())
+	{
+		FStructSchemaId FromSuper = ToStructIds[From.GetSuperSchema().Get().Idx];
+		FStructSchemaId ToSuper = static_cast<FStructSchemaId>(To.GetInnerSchemas()[0]);
+		return FromSuper == ToSuper;
+	}
+
+	return true;
 }
 
 ////////////////////////////////////////////////////////////////////////////
@@ -103,34 +142,97 @@ struct FLoadBatch
 
 using SubsetByteArray = TArray<uint8, TInlineAllocator<1024>>;
 
-static void CreateSubsetSchema(const FStructSchema& From, const FStructSchemaBinding& To, SubsetByteArray& Out)
+static void CopyMemberBinding(FLeafMemberBinding From, FMemberBinder& To)
 {
-	unimplemented();
+	// Skip enum schema
+	To.AddMember(From.Leaf.Pack(), static_cast<uint32>(From.Offset));
 }
 
-static FLoadStructPlan CreateLoadPlan(const FStructSchema& From, const FStructSchemaBinding& To, SubsetByteArray& OutSubsetSchemas)
+static void CopyMemberBinding(FStructMemberBinding From, FMemberBinder& To)
 {
-	// todo: detect if memcpyable, return FLoadStructMemcpy plan
+	To.AddMember(FMemberBindType(From.Type), static_cast<uint32>(From.Offset));
+	To.AddInnerSchema(From.Id);
+}
 
-	if (From.NumMembers < To.NumMembers || ContainsEnums(To))
+static void CopyMemberBinding(FRangeMemberBinding From, FMemberBinder& To)
+{
+	FMemberBindType InnermostType = From.InnerTypes[From.NumRanges - 1];
+	To.AddRange(MakeArrayView(From.RangeBindings, From.NumRanges), InnermostType, static_cast<uint32>(From.Offset));
+	To.AddOptionalInnerSchema(From.InnermostSchema);
+}
+
+static void CopyMemberBinding(/* in-out */ FMemberVisitor& From, FMemberBinder& To)
+{
+	switch (From.PeekKind())
 	{
-		CreateSubsetSchema(From, To, /* out */ OutSubsetSchemas);
+		case EMemberKind::Leaf:		CopyMemberBinding(From.GrabLeaf(), To);		break;
+		case EMemberKind::Range:	CopyMemberBinding(From.GrabRange(), To);	break;
+		case EMemberKind::Struct:	CopyMemberBinding(From.GrabStruct(), To);	break;
+		default:					check(false);								break;
 	}
-	else
+}
+
+static void CreateSubsetLoadSchema(const FStructSchema& From, const FStructSchemaBinding& To, TConstArrayView<FMemberId> ToNames,  uint16 NumEnums, SubsetByteArray& Out)
+{
+	check(To.NumMembers == ToNames.Num());
+	check(To.NumMembers >= From.NumMembers);
+
+	int32 OutPos = Out.Num();
+	FStructSchemaBinding Header = { From.NumMembers, From.NumInnerSchemas - NumEnums, From.NumRangeTypes };
+	Out.AddUninitialized(Header.CalculateSize());
+	FStructSchemaBinding* Schema = new (&Out[OutPos]) FStructSchemaBinding {Header};
+	
+	FMemberVisitor ToIt(To);
+	FMemberBinder Footer(*Schema);
+	for (FMemberId FromName : From.GetMemberNames())
+	{
+		while (FromName != ToNames[ToIt.GetIndex()])
+		{
+			ToIt.SkipMember();
+		}
+
+		CopyMemberBinding(/* in-out */ ToIt, Footer);			
+	}
+
+	unimplemented(); // todo copy kept members but skip enum schemas
+}
+
+[[nodiscard]] static FLoadStructPlan MakeSchemaLoadPlan(const FStructSchema& From, const FStructSchemaBinding& To, TConstArrayView<FMemberId> ToMemberIds, TConstArrayView<FStructSchemaId> ToStructIds, SubsetByteArray& OutSubsetSchemas)
+{
+	uint16 NumEnums = CountEnums(From);
+	if (From.NumMembers < To.NumMembers || NumEnums || HasDifferentSupers(From, To, ToStructIds))
+	{
+		CreateSubsetLoadSchema(From, To, ToMemberIds, NumEnums, /* out */ OutSubsetSchemas);
+	}
+	else // Reuse To bindings
 	{
 		check(From.NumMembers == To.NumMembers);
 		check(From.NumInnerSchemas == To.NumInnerSchemas);
 		check(From.NumRangeTypes == To.NumInnerRanges);
 	}
 
-	// New subset schema pointers will be remapped later
-	return FLoadStructPlan(To, ELeafWidth::B32, From.IsDense);
+	// Pointer to created subset load schema will be remapped later
+	return FLoadStructPlan(To, ELeafWidth::B32, !From.IsDense);
+}
+
+[[nodiscard]] static TOptional<FLoadStructMemcpy> TryMakeMemcpyPlan(const FStructSchema& From, const FStructSchemaBinding& To, TConstArrayView<FMemberId> ToMemberIds)
+{
+	// todo
+	return NullOpt;
+}
+
+[[nodiscard]] static FLoadStructPlan MakeLoadPlan(const FStructSchema& From, const FStructSchemaBinding& To, TConstArrayView<FMemberId> ToMemberIds, TConstArrayView<FStructSchemaId> ToStructIds, SubsetByteArray& OutSubsetSchemas)
+{
+	TOptional<FLoadStructMemcpy> Memcpy = TryMakeMemcpyPlan(From, To, ToMemberIds);
+	return Memcpy ? FLoadStructPlan(Memcpy.GetValue()) : MakeSchemaLoadPlan(From, To, ToMemberIds, ToStructIds, OutSubsetSchemas);
 }
 
-FLoadBatch* CreateLoadPlans(FReadBatchId ReadId, const FDeclarations& Declarations, const FStructBindings& Bindings)
+FLoadBatch* CreateLoadPlans(FReadBatchId ReadId, const FDeclarations& Declarations, const FStructBindings& Bindings, TConstArrayView<FStructSchemaId> LoadIds)
 {
+	check(NumStructSchemas(ReadId) == LoadIds.Num());
+
 	// Temporary data structures
-	const uint32 NumPlans = NumStructSchemas(ReadId);
+	const uint32 NumPlans = LoadIds.Num();
 	TArray<FLoadStructPlan, TInlineAllocator<256>> Plans;
 	TArray<uint32, TInlineAllocator<256>> SubsetSchemaSizes;
 	SubsetByteArray SubsetSchemaData;
@@ -138,14 +240,25 @@ FLoadBatch* CreateLoadPlans(FReadBatchId ReadId, const FDeclarations& Declaratio
 	SubsetSchemaSizes.SetNumUninitialized(NumPlans);
 
 	// Create plans
-	for (FStructSchemaId Id = { 0 }; Id.Idx < NumPlans; ++Id.Idx)
+	for (FStructSchemaId SavedId = { 0 }; SavedId.Idx < NumPlans; ++SavedId.Idx)
 	{
-		FStructBinding Binding = Bindings.Get(Id);
+		FStructSchemaId LoadId = LoadIds[SavedId.Idx];
+		FStructBinding Binding = Bindings.Get(LoadId);
 		int32 SubsetSchemaOffset = SubsetSchemaData.Num();
-		Plans[Id.Idx]	= Binding.IsCustom()
-						? FLoadStructPlan(Binding.AsCustom()) 
-						: CreateLoadPlan(ResolveStructSchema(ReadId, Id), Binding.AsSchema(), /* out */ SubsetSchemaData);
-		SubsetSchemaSizes[Id.Idx] = SubsetSchemaData.Num() - SubsetSchemaOffset;
+		if (Binding.IsCustom())
+		{
+			Plans[SavedId.Idx] = FLoadStructPlan(Binding.AsCustom()) ;
+		}
+		else
+		{
+			const FStructSchema& From = ResolveStructSchema(ReadId, LoadId);
+			const FStructSchemaBinding& To = Binding.AsSchema();
+			// Possible optimization - some simple memcpy cases doesn't need to resolve the declaration
+			TConstArrayView<FMemberId> ToMemberIds = Declarations.Get(LoadId).GetMemberOrder();
+			Plans[SavedId.Idx] = MakeLoadPlan(From, To, ToMemberIds, LoadIds, /* out */ SubsetSchemaData);	
+		}
+		
+		SubsetSchemaSizes[SavedId.Idx] = SubsetSchemaData.Num() - SubsetSchemaOffset;
 	}
 	
 	// Allocate load batch, copy plans and subset schemas, and fixup subset schema plans
@@ -158,7 +271,7 @@ FLoadBatch* CreateLoadPlans(FReadBatchId ReadId, const FDeclarations& Declaratio
 		uint8* OutSubsetData = reinterpret_cast<uint8*>(Out->Plans + Plans.Num());
 		FMemory::Memcpy(OutSubsetData, SubsetSchemaData.GetData(), SubsetSchemaData.Num());
 		
-		// Remap subset schema pointers
+		// Update plans with actual subset schema pointers
 		const uint8* It = OutSubsetData;
 		for (uint32 Idx = 0; Idx < NumPlans; ++Idx)
 		{
@@ -177,6 +290,11 @@ FLoadBatch* CreateLoadPlans(FReadBatchId ReadId, const FDeclarations& Declaratio
 	return Out;
 }
 
+void DestroyLoadPlans(FLoadBatch* Batch)
+{
+	FMemory::Free(Batch);
+}
+
 ////////////////////////////////////////////////////////////////////////////
 
 FORCEINLINE static void SetBit(uint8& Out, uint8 Idx, bool bValue)
@@ -340,7 +458,8 @@ public:
 	, Batch(InBatch)
 	, ByteIt(Values)
 	{
-		check(!ContainsEnums(Schema));
+		checkf(CountStaticStructs(Schema) == Schema.NumInnerSchemas, TEXT("Expects Schema stripped from load-irrelevant enum schema ids, see CreateSubsetLoadSchema(). "
+			"# schemas/structs/enums: %d/%d/%d"), Schema.NumInnerSchemas, CountStaticStructs(Schema), CountEnums(Schema));
 	}
 
 	void Load(uint8* Struct)
diff --git a/Engine/Source/Programs/PlainPropsTests/Private/PlainPropsRead.cpp b/Engine/Source/Programs/PlainPropsTests/Private/PlainPropsRead.cpp
--- a/Engine/Source/Programs/PlainPropsTests/Private/PlainPropsRead.cpp
+++ b/Engine/Source/Programs/PlainPropsTests/Private/PlainPropsRead.cpp
@@ -1,7 +1,6 @@
 // Copyright Epic Games, Inc. All Rights Reserved.
 
 #include  "PlainPropsRead.h"
-#include  "PlainPropsInternalFormat.h"
 #include  "PlainPropsInternalRead.h"
 #include "Misc/Optional.h"
 
@@ -18,7 +17,7 @@ void FSchemaBatch::ValidateBounds(uint64 NumBytes) const
 	check(sizeof(FSchemaBatch) + NumSchemas * sizeof(uint32) <= NestedScopesOffset);
 	check(NestedScopesOffset + NumNestedScopes * sizeof(FNestedScope) + NumParametricTypes * sizeof(FParametricType) <= NumBytes);
 
-	for (uint32 SchemaOffset : MakeArrayView(SchemaOffsets, NumSchemas))
+	for (uint32 SchemaOffset : GetSchemaOffsets())
 	{
 		check(SchemaOffset < NestedScopesOffset);
 		check(IsAligned(SchemaOffset, Alignment));
@@ -118,7 +117,7 @@ const FSchemaBatch* UnmountReadSchemas(FReadBatchId Id)
 
 uint32 NumStructSchemas(FReadBatchId Batch)
 {
-	return GReadSchemas.Get(Batch).NumSchemas;
+	return GReadSchemas.Get(Batch).NumStructSchemas;
 }
 
 const FStructSchema& ResolveStructSchema(FReadBatchId Batch, FStructSchemaId Schema)
diff --git a/Engine/Source/Programs/PlainPropsTests/Private/PlainPropsSave.cpp b/Engine/Source/Programs/PlainPropsTests/Private/PlainPropsSave.cpp
--- a/Engine/Source/Programs/PlainPropsTests/Private/PlainPropsSave.cpp
+++ b/Engine/Source/Programs/PlainPropsTests/Private/PlainPropsSave.cpp
@@ -259,8 +259,9 @@ template<class SaverType, typename InnerContextType>
 
 [[nodiscard]] static FRangeMemberBinding GetInnerRange(FRangeMemberBinding Member)
 {
+	check(Member.NumRanges > 1);
 	check(Member.InnerTypes[0].IsRange());
-	return { Member.InnerTypes + 1, Member.RangeBindings + 1, Member.InnermostSchema };
+	return { Member.InnerTypes + 1, Member.RangeBindings + 1, Member.NumRanges - 1, Member.InnermostSchema };
 }
 
  ELeafWidth GetArithmeticWidth(FLeafBindType Leaf)
diff --git a/Engine/Source/Programs/PlainPropsTests/Private/PlainPropsWrite.cpp b/Engine/Source/Programs/PlainPropsTests/Private/PlainPropsWrite.cpp
--- a/Engine/Source/Programs/PlainPropsTests/Private/PlainPropsWrite.cpp
+++ b/Engine/Source/Programs/PlainPropsTests/Private/PlainPropsWrite.cpp
@@ -31,17 +31,17 @@ struct FWriteIds
 	TArray<FParametricType>				KeptParametrics;
 	TArray<FTypeId>						KeptParameters;
 
-	FNameId					Remap(FNameId Old) const				{ return Names[Old.Idx].Get(); }
-	FMemberId				Remap(FMemberId Old) const				{ return { Remap(Old.Id) }; }
-	FFlatScopeId			Remap(FFlatScopeId Old) const			{ return { Remap(Old.Name) }; }
-	FNestedScopeId			Remap(FNestedScopeId Old) const			{ return NestedScopes[Old.Idx].Get(); }
-	FScopeId				Remap(FScopeId Old) const				{ return Old.IsNested() ? FScopeId(Remap(Old.AsNested())) : FScopeId(Remap(Old.AsFlat())); }
-	FConcreteTypenameId		Remap(FConcreteTypenameId Old) const	{ return { Remap(Old.Id) }; }	
-	FParametricTypeId		Remap(FParametricTypeId Old) const		{ return ParametricTypes[Old.Idx].Get(); }
-	FTypenameId				Remap(FTypenameId Old) const			{ return Old.IsConcrete() ? FTypenameId(Remap(Old.AsConcrete())) : FTypenameId(Remap(Old.AsParametric())); }
-	FTypeId					Remap(FTypeId Old) const				{ return { Remap(Old.Scope), Remap(Old.Name) }; }
-	FSchemaId				RemapStruct(FSchemaId Old) const		{ return Structs[Old.Idx].Get(); }
-	FSchemaId				RemapEnum(FSchemaId Old) const			{ return Enums[Old.Idx].Get(); }
+	FNameId								Remap(FNameId Old) const				{ return Names[Old.Idx].Get(); }
+	FMemberId							Remap(FMemberId Old) const				{ return { Remap(Old.Id) }; }
+	FFlatScopeId						Remap(FFlatScopeId Old) const			{ return { Remap(Old.Name) }; }
+	FNestedScopeId						Remap(FNestedScopeId Old) const			{ return NestedScopes[Old.Idx].Get(); }
+	FScopeId							Remap(FScopeId Old) const				{ return Old.IsNested() ? FScopeId(Remap(Old.AsNested())) : FScopeId(Remap(Old.AsFlat())); }
+	FConcreteTypenameId					Remap(FConcreteTypenameId Old) const	{ return { Remap(Old.Id) }; }	
+	FParametricTypeId					Remap(FParametricTypeId Old) const		{ return ParametricTypes[Old.Idx].Get(); }
+	FTypenameId							Remap(FTypenameId Old) const			{ return Old.IsConcrete() ? FTypenameId(Remap(Old.AsConcrete())) : FTypenameId(Remap(Old.AsParametric())); }
+	FTypeId								Remap(FTypeId Old) const				{ return { Remap(Old.Scope), Remap(Old.Name) }; }
+	FSchemaId							RemapStruct(FSchemaId Old) const		{ return Structs[Old.Idx].Get(); }
+	FSchemaId							RemapEnum(FSchemaId Old) const			{ return Enums[Old.Idx].Get(); }
 };
 
 static TConstArrayView<FNameId> GetUsedNames(const FBuiltStructSchema& Used)
@@ -467,6 +467,7 @@ static void WriteSchemasImpl(TArray64<uint8>& Out, const FBuiltSchemas& Schemas,
 	Header.NestedScopesOffset = IntCastChecked<uint32>(NestedScopePos - HeaderPos);
 	Header.NumParametricTypes = NewIds.KeptParametrics.Num();
 	Header.NumSchemas = NewIds.NumKeptSchemas;
+	Header.NumStructSchemas = Schemas.Structs.Num();
 	FMemory::Memcpy(&Out[HeaderPos], &Header, sizeof(FSchemaBatch));
 
 	// Write schema offsets
diff --git a/Engine/Source/Programs/PlainPropsTests/Public/PlainPropsBind.h b/Engine/Source/Programs/PlainPropsTests/Public/PlainPropsBind.h
--- a/Engine/Source/Programs/PlainPropsTests/Public/PlainPropsBind.h
+++ b/Engine/Source/Programs/PlainPropsTests/Public/PlainPropsBind.h
@@ -4,6 +4,7 @@
 
 #include "Containers/ArrayView.h"
 #include "Memory/MemoryFwd.h"
+#include "Memory/MemoryView.h"
 #include  "PlainPropsTypes.h"
 #include  "PlainPropsCtti.h"
 #include  "PlainPropsDeclare.h"
@@ -11,12 +12,11 @@
 namespace PlainProps 
 {
 
-// Todo: Include Build.h and Load.h instead?
-enum class EMemberPresence;
-enum class EEnumMode;
 struct FBuiltStruct;
+class FIdIndexerBase;
 struct FLoadBatch;
 class FMemberBuilder;
+struct FSchemaBatch;
 class FStructBinding;
 struct FStructView;
 class FRangeBinding;
@@ -24,31 +24,6 @@ struct FTypedRange;
 class IRangeBinding;
 template<class T> class TIdIndexer;
 
-//////////////////////////////////////////////////////////////////////////
-
-struct FIdBinding
-{
-	TConstArrayView<FNameId>			Names;
-	TConstArrayView<FNestedScopeId>		Scopes;
-	TConstArrayView<FParametricTypeId>	Parametrics;
-};
-
-class FIdBinder
-{
-	const uint32 NumNames;
-	const uint32 NumScopes;
-	const uint32 NumParameters;
-	const FNameId* Data;
-
-public:
-	template<class NameType>
-	FIdBinder(TIdIndexer<NameType>& Indexer, TConstArrayView<NameType> Names, TConstArrayView<FNestedScope> Scopes, TConstArrayView<FParametricType> ParametericTypes);
-
-	FIdBinding Get() const;
-};
-
-
-
 ////////////////////////////////////////////////////////////////////////////////////////////////
 
 enum class ELeafBindType : uint8 { Bool, IntS, IntU, Float, Hex, Enum, Unicode, BitfieldBool };
@@ -79,8 +54,11 @@ struct FBitfieldBindType
 
 union FLeafBindType
 {
-	constexpr explicit FLeafBindType(FLeafType In) : Arithmetic({EMemberKind::Leaf, ToLeafBindType(In.Type), In.Width}) {}
+	constexpr explicit FLeafBindType(ELeafBindType ArithmeticType, ELeafWidth Width) : Arithmetic({EMemberKind::Leaf, ArithmeticType, Width}) {}
+	constexpr explicit FLeafBindType(FUnpackedLeafType In) : Arithmetic({EMemberKind::Leaf, ToLeafBindType(In.Type), In.Width}) {}
+	constexpr explicit FLeafBindType(FLeafType In) : FLeafBindType(FUnpackedLeafType(In)) {}
 	constexpr explicit FLeafBindType(FBitfieldBindType In) : Bitfield({EMemberKind::Leaf, ELeafBindType::BitfieldBool, In.Idx}) {}
+	constexpr explicit FLeafBindType(uint8 BitfieldIdx) : Bitfield({EMemberKind::Leaf, ELeafBindType::BitfieldBool, BitfieldIdx}) {}
 
 	struct
 	{
@@ -108,6 +86,8 @@ struct FStructBindType : FStructType {};
 union FMemberBindType
 {
 	constexpr explicit FMemberBindType(FLeafType In) : Leaf(In) {}
+	constexpr explicit FMemberBindType(FUnpackedLeafType In) : Leaf(In) {}
+	constexpr explicit FMemberBindType(FLeafBindType In) : Leaf(In) {}
 	constexpr explicit FMemberBindType(FBitfieldBindType In) : Leaf(In) {}
 	constexpr explicit FMemberBindType(FRangeType In) : Range(In) {}
 	constexpr explicit FMemberBindType(ERangeSizeType MaxSize) : Range({EMemberKind::Range, MaxSize}) {}
@@ -148,8 +128,12 @@ struct FStructSchemaBinding
 	const uint32*			GetOffsets() const			{ return AlignPtr<uint32>(GetInnerRangeTypes() + NumInnerRanges); }
 	const FSchemaId*		GetInnerSchemas() const		{ return AlignPtr<FSchemaId>(GetOffsets() + NumMembers); }
 	const FRangeBinding*	GetRangeBindings() const	{ return AlignPtr<FRangeBinding>(GetInnerSchemas() + NumInnerSchemas); }
+	uint32					CalculateSize() const;
+	bool					HasSuper() const			{ return NumInnerSchemas > 0 && Members[0].IsStruct() && Members[0].AsStruct().IsSuper; }
 };
 
+////////////////////////////////////////////////////////////////////////////////////////////////
+
 struct FUnpackedLeafBindType
 {
 	ELeafBindType			Type;
@@ -174,7 +158,10 @@ struct FUnpackedLeafBindType
 	}
 
 //	constexpr bool operator==(FUnpackedLeafBindType O) { return Type == O.Type && Width == O.Width; }
-//	FMemberBindType Pack() const { return FMemberBindType(Type, Width); }
+	FMemberBindType Pack() const
+	{ 
+		return FMemberBindType(Type == ELeafBindType::BitfieldBool ? FLeafBindType(BitfieldIdx) : FLeafBindType(Type, Width));
+	}
 };
 
 struct FLeafMemberBinding
@@ -188,12 +175,14 @@ struct FRangeMemberBinding
 {
 	const FMemberBindType*	InnerTypes;
 	const FRangeBinding*	RangeBindings;
+	uint32					NumRanges; // At least 1, >1 for nested ranges
 	FOptionalSchemaId		InnermostSchema;
 	uint64					Offset;
 };
 
 struct FStructMemberBinding
 {
+	FStructType			Type;
 	FStructSchemaId		Id;
 	uint64				Offset;
 };
@@ -212,21 +201,15 @@ public:
 	FLeafMemberBinding			GrabLeaf();				// @pre PeekKind() == EMemberKind::Leaf
 	FRangeMemberBinding			GrabRange();			// @pre PeekKind() == EMemberKind::Range
 	FStructMemberBinding		GrabStruct();			// @pre PeekKind() == EMemberKind::Struct
+	void						SkipMember();
 
 protected: // for unit tests
 	const FStructSchemaBinding& Schema;
-	//const FMemberBindType*	Footer;
 	const uint16				NumMembers;
-	//const uint16			NumInnerRanges;			// Number of ranges and nested ranges
-	//const uint16			NumInnerSchemas;		// Number of static structs and enums
 	uint16						MemberIdx = 0;
 	uint16						InnerRangeIdx = 0;		// Types of [nested] ranges
 	uint16						InnerSchemaIdx = 0;		// Types of static structs and enums
 
-	//const uint32*				GetOffsets() const;
-	//const FMemberBindType*		GetInnerRanges() const;
-	//const FSchemaId*			GetInnerSchemas() const;
-
 	using FMemberBindTypeRange = TConstArrayView<FMemberBindType>;
 
 	uint64						GrabMemberOffset();
@@ -430,7 +413,9 @@ using RangeBind = typename TRangeBind<T>::Type;
 
 struct FMemberBinding
 {
-	uint64							Offset;
+	FMemberBinding() : InnermostType(FLeafBindType(ELeafBindType::Bool, ELeafWidth::B8)) {}
+
+	uint64							Offset = 0;
 	FMemberBindType					InnermostType;		// Always Leaf or Struct
 	FOptionalSchemaId				InnermostSchema;	// Enum or struct schema
 	TConstArrayView<FRangeBinding>	RangeBindings;		// Non-empty -> Range
@@ -438,7 +423,6 @@ struct FMemberBinding
 
 ////////////////////////////////////////////////////////////////////////////////////////////////
 
-
 class FStructBindings
 {
 public:
@@ -459,19 +443,40 @@ private:
 
 ////////////////////////////////////////////////////////////////////////////////////////////////
 
-//struct FNamedMemberBinding : FMemberBinding
-//{
-//	FMemberId					Name;
-//};
+template<typename Enum, typename Ids>
+FEnumSchemaId IndexNativeEnum()
+{
+	static FEnumSchemaId Id = Ids::IndexEnum(CttiOf<Enum>::Name);
+	return Id;
+}
+
+template<typename Struct, typename Ids>
+FStructSchemaId IndexNativeStruct()
+{
+	static FStructSchemaId Id = Ids::IndexStruct(CttiOf<Struct>::Name);
+	return Id;
+}
+
+template<typename Struct, typename Ids>
+FOptionalStructSchemaId IndexOptionalNativeStruct()
+{
+	if constexpr (!std::is_void_v<Struct>)
+	{
+		return IndexNativeStruct<Struct, Ids>();
+	}
+	
+	return NoId;
+}
+
+////////////////////////////////////////////////////////////////////////////////////////////////
 
 template<typename Type, class Ids>
 FMemberBindType BindMemberLeaf(FOptionalSchemaId& OutSchema)
 {
-	OutSchema = std::is_enum_v<Type> ? IndexNativeEnum<Type, Ids>() : NoId;
-	return ReflectLeaf<Type>;
+	OutSchema = std::is_enum_v<Type> ? ToOptional(static_cast<FSchemaId>(IndexNativeEnum<Type, Ids>())) : NoId;
+	return FMemberBindType(ReflectLeaf<Type>);
 }
 
-
 template<typename CustomBinding, class Runtime>
 FStructSchemaId BindCustomStructOnce()
 {
@@ -578,8 +583,8 @@ FMemberBinding BindMember()
 	using Ids = typename Runtime::Ids;
 	using Type = Var::Type;
 
-	FMemberBindType Dummy(ERangeSizeType);
-	FMemberBinding Out = { Var::Offset, Dummy };
+	FMemberBinding Out;
+	Out.Offset = Var::Offset;
 	if constexpr (std::is_arithmetic_v<Type> || std::is_enum_v<Type>)
 	{
 		Out.InnermostType = BindMemberLeaf<Type, Ids>(Out.InnermostSchema);
@@ -605,32 +610,7 @@ FMemberBinding BindMember()
 	return Out;
 }
 
-template<typename Enum, typename Ids>
-FEnumSchemaId IndexNativeEnum()
-{
-	static FEnumSchemaId Id = Ids::IndexEnum(CttiOf<Enum>::Name);
-	return Id;
-}
-
-template<typename Struct, typename Ids>
-FStructSchemaId IndexNativeStruct()
-{
-	static FStructSchemaId Id = Ids::IndexStruct(CttiOf<Struct>::Name);
-	return Id;
-}
-
-
-template<typename Struct, typename Ids>
-FOptionalStructSchemaId IndexOptionalNativeStruct()
-{
-	if constexpr (!std::is_void_v<Struct>)
-	{
-		return IndexNativeStruct<Struct, Ids>();
-	}
-	
-	return NoId;
-}
-
+////////////////////////////////////////////////////////////////////////////////////////////////
 
 template<class Ctti, class Ids>
 FEnumSchemaId DeclareNativeEnum(FDeclarations& Out, EEnumMode Mode)
@@ -698,5 +678,127 @@ void BindNativeStruct(FStructBindings& Out, FStructSchemaId DeclaredId)
 
 ////////////////////////////////////////////////////////////////////////////////////////////////
 
+struct FMemberBinder
+{
+	FMemberBinder(FStructSchemaBinding& InSchema)
+	: Schema(InSchema)
+	, MemberIt(Schema.Members)
+	, RangeTypeIt(const_cast<FMemberBindType*>(Schema.GetInnerRangeTypes()))
+	, OffsetIt(const_cast<uint32*>(Schema.GetOffsets()))
+	, InnerSchemaIt(const_cast<FSchemaId*>(Schema.GetInnerSchemas()))
+	, RangeBindingIt(const_cast<FRangeBinding*>(Schema.GetRangeBindings()))
+	{}
+
+	~FMemberBinder()
+	{
+		check(MemberIt == Schema.GetInnerRangeTypes());
+		check(Align(RangeTypeIt, alignof(uint32)) == (const void*)Schema.GetOffsets());
+		check(OffsetIt == (const void*)Schema.GetInnerSchemas());
+		check(Align(InnerSchemaIt, alignof(FRangeBinding)) == (const void*)Schema.GetRangeBindings());
+		check(Schema.NumInnerRanges == RangeBindingIt - Schema.GetRangeBindings());
+	}
+
+	void AddMember(FMemberBindType Type, uint32 Offset)
+	{
+		*MemberIt++ = Type;
+		*OffsetIt++ = Offset;
+	}
+
+	void AddRange(TConstArrayView<FRangeBinding> Ranges, FMemberBindType InnermostType, uint32 Offset)
+	{
+		AddMember(FMemberBindType(Ranges[0].GetSizeType()), Offset);
+
+		for (FRangeBinding Range : Ranges.RightChop(1))
+		{
+			*RangeTypeIt++ = FMemberBindType(Range.GetSizeType());
+		}
+		*RangeTypeIt++ = InnermostType;
+
+		FMemory::Memcpy(RangeBindingIt, Ranges.GetData(), Ranges.Num() * Ranges.GetTypeSize());
+		RangeBindingIt += Ranges.Num();
+	}
+	
+	void AddInnerSchema(FSchemaId InnermostSchema)
+	{
+		*InnerSchemaIt++ = InnermostSchema;
+	}
+
+	void AddOptionalInnerSchema(FOptionalSchemaId InnermostSchema)
+	{
+		if (InnermostSchema)
+		{
+			AddInnerSchema(InnermostSchema.Get());
+		}
+	}
+
+	FStructSchemaBinding& Schema;
+	FMemberBindType* MemberIt;
+	FMemberBindType* RangeTypeIt;
+	uint32* OffsetIt;
+	FSchemaId* InnerSchemaIt;
+	FRangeBinding* RangeBindingIt;
+};
+
+//////////////////////////////////////////////////////////////////////////
+
+// Save -> load ids
+struct FIdBinding
+{
+	TConstArrayView<FNameId>			Names;
+	TConstArrayView<FNestedScopeId>		NestedScopes;
+	TConstArrayView<FParametricTypeId>	ParametricTypes;
+	TConstArrayView<FSchemaId>			Schemas;
+
+	FNameId								Remap(FNameId Old) const				{ return Names[Old.Idx]; }
+	FMemberId							Remap(FMemberId Old) const				{ return { Remap(Old.Id) }; }
+	FFlatScopeId						Remap(FFlatScopeId Old) const			{ return { Remap(Old.Name) }; }
+	FNestedScopeId						Remap(FNestedScopeId Old) const			{ return NestedScopes[Old.Idx]; }
+	FScopeId							Remap(FScopeId Old) const				{ return Old.IsNested() ? FScopeId(Remap(Old.AsNested())) : FScopeId(Remap(Old.AsFlat())); }
+	FConcreteTypenameId					Remap(FConcreteTypenameId Old) const	{ return { Remap(Old.Id) }; }	
+	FParametricTypeId					Remap(FParametricTypeId Old) const		{ return ParametricTypes[Old.Idx]; }
+	FTypenameId							Remap(FTypenameId Old) const			{ return Old.IsConcrete() ? FTypenameId(Remap(Old.AsConcrete())) : FTypenameId(Remap(Old.AsParametric())); }
+	FTypeId								Remap(FTypeId Old) const				{ return { Remap(Old.Scope), Remap(Old.Name) }; }
+
+	TConstArrayView<FStructSchemaId>	GetStructIds(int32 NumStructs) const
+	{
+		// All saved struct schema ids are lower than enum schema ids
+		check(NumStructs <= Schemas.Num());
+		return MakeArrayView(static_cast<const FStructSchemaId*>(Schemas.GetData()), NumStructs);
+	}
+};
+
+struct FIdTranslatorBase
+{
+	static uint32 CalculateTranslationSize(int32 NumSavedNames, const FSchemaBatch& Batch);
+	static FIdBinding TranslateIds(FMutableMemoryView To, FIdIndexerBase& Indexer, TConstArrayView<FNameId> TranslatedNames, const FSchemaBatch& From);
+};
+
+// Maps saved ids -> runtime load ids
+struct FIdTranslator : FIdTranslatorBase
+{
+	template<class NameType>
+	FIdTranslator(TIdIndexer<NameType>& Indexer, TConstArrayView<NameType> SavedNames, const FSchemaBatch& Batch)
+	{
+		Allocator.SetNumUninitialized(CalculateTranslationSize(SavedNames.Num(), Batch));
+		
+		// Translate names
+		TArrayView<FNameId> NewNames(reinterpret_cast<FNameId*>(&Allocator[0]), SavedNames.Num());
+		FNameId* NameIt = &NewNames[0];
+		for (const NameType& SavedName : SavedNames)
+		{
+			(*NameIt++) = Indexer.MakeName(SavedName);
+		}
+
+		FMutableMemoryView OtherIds(NameIt, Allocator.Num() - NewNames.Num() * sizeof(FNameId));
+		Translation = TranslateIds(/* out */ OtherIds, Indexer, NewNames, Batch);
+	}
+	
+	FIdBinding								Translation;
+	TArray<uint8, TInlineAllocator<1024>>	Allocator;
+};
+
+FSchemaBatch*			CreateTranslatedSchemas(const FSchemaBatch& Schemas, FIdBinding NewIds);
+void					DestroyTranslatedSchemas(const FSchemaBatch* Schemas);
+
 } // namespace PlainProps
 
diff --git a/Engine/Source/Programs/PlainPropsTests/Public/PlainPropsIndex.h b/Engine/Source/Programs/PlainPropsTests/Public/PlainPropsIndex.h
--- a/Engine/Source/Programs/PlainPropsTests/Public/PlainPropsIndex.h
+++ b/Engine/Source/Programs/PlainPropsTests/Public/PlainPropsIndex.h
@@ -50,6 +50,7 @@ class FIdIndexerBase : public FDebugIds
 {
 public:
 	FScopeId							NestScope(FScopeId Outer, FFlatScopeId Inner);
+	FParametricTypeId					MakeParametricTypeId(FConcreteTypenameId Name, TConstArrayView<FTypeId> Params);
 	FTypeId								MakeParametricType(FTypeId Type, TConstArrayView<FTypeId> Params);
 
 	FEnumSchemaId						IndexEnum(FTypeId Type);
diff --git a/Engine/Source/Programs/PlainPropsTests/Public/PlainPropsLoad.h b/Engine/Source/Programs/PlainPropsTests/Public/PlainPropsLoad.h
--- a/Engine/Source/Programs/PlainPropsTests/Public/PlainPropsLoad.h
+++ b/Engine/Source/Programs/PlainPropsTests/Public/PlainPropsLoad.h
@@ -16,7 +16,7 @@ class FStructBindings;
 struct FTypedRange;
 
 
-PP_API FLoadBatch* CreateLoadPlans(FReadBatchId ReadId, const FDeclarations& Declarations, const FStructBindings& Bindings);
+PP_API FLoadBatch* CreateLoadPlans(FReadBatchId ReadId, const FDeclarations& Declarations, const FStructBindings& Bindings, TConstArrayView<FStructSchemaId> LoadIds);
 PP_API void DestroyLoadPlans(FLoadBatch* Batch);
 PP_API void LoadStruct(uint8* Dst, FByteReader Src, FStructSchemaId Id, const FLoadBatch& Batch);
 PP_API void LoadStruct(uint8* Dst, FStructView Src, const FLoadBatch& Batch);
diff --git a/Engine/Source/Programs/PlainPropsTests/Public/PlainPropsRead.h b/Engine/Source/Programs/PlainPropsTests/Public/PlainPropsRead.h
--- a/Engine/Source/Programs/PlainPropsTests/Public/PlainPropsRead.h
+++ b/Engine/Source/Programs/PlainPropsTests/Public/PlainPropsRead.h
@@ -13,7 +13,6 @@ namespace PlainProps
 {
 
 struct FEnumSchema;
-struct FIdBinding;
 struct FSchemaBatch;
 struct FStructSchema;
 
@@ -28,10 +27,8 @@ class FReadBatchId
 
 // @param Schemas must outlive read batch
 const FSchemaBatch*		ValidateSchemas(FMemoryView Schemas);
-FSchemaBatch*			CreateTranslatedSchemas(const FSchemaBatch* Schemas, FIdBinding NewIds);
 FReadBatchId			MountReadSchemas(const FSchemaBatch* Schemas);
 const FSchemaBatch*		UnmountReadSchemas(FReadBatchId Batch);
-void					DestroyTranslatedSchemas(const FSchemaBatch* Schemas);
 
 uint32					NumStructSchemas(FReadBatchId Batch);
 const FStructSchema&	ResolveStructSchema(FReadBatchId Batch, FStructSchemaId Id);
@@ -371,10 +368,10 @@ bool EqualItems(RangeTypeA&& A, RangeTypeB&& B)
 class FLeafRangeView
 {
 public:
-	ELeafType GetType() const;
-	SIZE_T GetSizeOf() const;
-	uint64 Num() const { return NumItems; }
-	uint64 NumBytes() const { return NumItems * GetSizeOf(); }
+	ELeafType				GetType() const;
+	SIZE_T					GetSizeOf() const;
+	uint64					Num() const { return NumItems; }
+	uint64					NumBytes() const { return NumItems * GetSizeOf(); }
 	
 	// These ranges hide the internal representations so we can change format in the future, 
 	// e.g. store zeroes or 1.0f in some compact fashion or even var int encodings
diff --git a/Engine/Source/Programs/PlainPropsTests/Tests/PlainPropsTests.cpp b/Engine/Source/Programs/PlainPropsTests/Tests/PlainPropsTests.cpp
--- a/Engine/Source/Programs/PlainPropsTests/Tests/PlainPropsTests.cpp
+++ b/Engine/Source/Programs/PlainPropsTests/Tests/PlainPropsTests.cpp
@@ -1,6 +1,6 @@
 // Copyright Epic Games, Inc. All Rights Reserved.
 
-//#if WITH_TESTS //TMPTORP
+#if WITH_TESTS
 #include "Tests/TestHarnessAdapter.h"
 #include "Containers/AnsiString.h"
 #include "Containers/Map.h"
@@ -1285,87 +1285,7 @@ TEST_CASE_NAMED(FPlainPropsReadWriteTest, "System::Core::Serialization::PlainPro
 	}
 }
 
-//
-//namespace Test
-//{
-//
-//class FBindings
-//{
-//
-//	
-//};
-//	
-//struct FPt
-//{
-//	float x, y;
-//};
-//UE_REFLECT_STRUCT(FPt, A, B);
-//
-//
-//enum EFlat { A = 1, B = 3 };
-//enum EFlag { A = 0, B = 2 };
-//
-//UE_REFLECT_ENUM(EFlat, A, B);
-//UE_REFLECT_ENUM(EFlag, A, B);
-//
-//struct FEnums
-//{
-//	EFlat Flat;
-//	EFlag Flag;
-//};
-//
-//UE_REFLECT_STRUCT(FEnums, Flat, Flag);
-//
-//struct Leaves
-//{
-//	uint16 U16;
-//	int64 I64;
-//	bool Bool;
-//	char AsciiChar;
-//	uint8 BitfieldBool : 1;
-//	uint32 StaticU32s[3];
-//};
-//
-//UE_REFLECT_STRUCT(Leaves, U16, I64, Bool, AsciiChar, BitfieldBool, StaticU32s);
-//
-//GTestBindings.Add(LeavesCtti);
-//
-//struct FNestedStructs
-//{
-//	FPt Point;
-//	FPt StaticPoints[2];
-//};
-//UE_REFLECT_STRUCT(FNestedStructs, Point, StaticPoints);
-//
-//struct FLeafRanges
-//{
-//	TArray<int32> IntArray;
-//	TOptional<uint8> MaybeByte;
-//	TUniquePtr<float> FloatPtr;
-//	TSet<uint16> ShortSet;
-//	TSparseArray<bool> SparseBools;
-//};
-//
-//struct FSuper
-//{
-//	uint16 Pad;
-//	bool A;
-//};
-//UE_REFLECT_STRUCT(FSuper, A);
-//
-//struct FSub : FSuper
-//{
-//	bool B;
-//	uint32 Pad;
-//};
-//UE_REFLECT_SUBSTRUCT(FSub, FSuper, B);
-//
-//
-//struct FSubs
-//{
-//	TArray<FSub> Subs;
-//};
-
+////////////////////////////////////////////////////////////////////////////
 
 TEST_CASE_NAMED(FPlainPropsLoadSaveTest, "System::Core::Serialization::PlainProps::LoadSave", "[Core][PlainProps][SmokeFilter]")
 {
@@ -1424,9 +1344,5 @@ TEST_CASE_NAMED(FPlainPropsLoadSaveTest, "System::Core::Serialization::PlainProp
 	{}
 }
 
-
-//} // namespace Test
-
-
 } // namespace PlainProps
-//#endif // WITH_TESTS
\ No newline at end of file
+#endif // WITH_TESTS
\ No newline at end of file
diff --git a/Engine/Source/Programs/PlainPropsTests/Tests/PlainPropsUeCoreTests.cpp b/Engine/Source/Programs/PlainPropsTests/Tests/PlainPropsUeCoreTests.cpp
--- a/Engine/Source/Programs/PlainPropsTests/Tests/PlainPropsUeCoreTests.cpp
+++ b/Engine/Source/Programs/PlainPropsTests/Tests/PlainPropsUeCoreTests.cpp
@@ -89,7 +89,7 @@ struct TScopedStructBinding
 	TScopedStructBinding(EMemberPresence Occupancy = EMemberPresence::AllowSparse)
 	: Id(DeclareNativeStruct<Ctti, Ids>(Runtime::GetDeclarations(), Occupancy))
 	{
-		//BindNativeStruct<Ctti, Runtime>(Runtime::GetBindings());
+		BindNativeStruct<Ctti, Runtime>(Runtime::GetBindings(), Id);
 		//if constexpr (bDeltaBindMembers)
 		//{
 		//	BindNativeStruct<Ctti, FDeltaRuntime>(Runtime::GetBindings());	
@@ -193,8 +193,6 @@ TArray64<uint8> FBatchSaver::Write() const
 	TArray64<uint8> Out;
 	WriteU32(Out, Magics[0]);
 	WriteNumAndArray(Out, TArrayView<const FName, int32>(UsedNames));
-	//WriteNumAndArray(Out, Writer.GetUsedScopes());
-	//WriteNumAndArray(Out, Writer.GetUsedParametrics());
 
 	// Write schemas
 	WriteU32(Out, Magics[1]);
@@ -204,7 +202,7 @@ TArray64<uint8> FBatchSaver::Write() const
 	WriteNumAndArray(Out, TArrayView<const uint8, int64>(Tmp));
 	Tmp.Reset();
 
-	//// Write objects
+	// Write objects
 	WriteU32(Out, Magics[2]);
 	for (const TPair<FStructSchemaId, TUniquePtr<FBuiltStruct>>& Object : SavedObjects)
 	{
@@ -232,18 +230,19 @@ public:
 		// Read ids
 		FByteReader It(Data);
 		CHECK(It.Grab<uint32>() == Magics[0]);
-		Names = GrabNumAndArray<FName>(It);
+		SavedNames = GrabNumAndArray<FName>(It);
+		CHECK(SavedNames.Num() != 0);
 		
 		// Read schemas
 		CHECK(It.Grab<uint32>() == Magics[1]);
 		It.SkipAlignmentPadding<uint32>();
 		uint32 SchemasSize = It.Grab<uint32>();
-		const FSchemaBatch* SaveSchemas = ValidateSchemas(It.GrabSlice(SchemasSize));
+		const FSchemaBatch* SavedSchemas = ValidateSchemas(It.GrabSlice(SchemasSize));
 		CHECK(It.Grab<uint32>() == Magics[2]);
 		
-		// Bind saved ids to current ids, remap schemas with new ids and mount it
-		FIdBinder SaveToLoadIds(GNames, Names, SaveSchemas->GetNestedScopes(), SaveSchemas->GetParametricTypes());
-		FSchemaBatch* LoadSchemas = CreateTranslatedSchemas(SaveSchemas, SaveToLoadIds.Get());
+		// Bind saved ids to runtime ids, make new schemas with new ids and mount them
+		FIdTranslator LoadIds(GNames, SavedNames, *SavedSchemas);
+		FSchemaBatch* LoadSchemas = CreateTranslatedSchemas(*SavedSchemas, LoadIds.Translation);
 		FReadBatchId Batch = MountReadSchemas(LoadSchemas);
 
 		// Read objects
@@ -259,28 +258,32 @@ public:
 		CHECK(!Objects.IsEmpty());
 
 		// Finally create load plans
-		Plans = CreateLoadPlans(Batch, GTypes, Bindings);
+		TConstArrayView<FStructSchemaId> LoadStructIds = LoadIds.Translation.GetStructIds(SavedSchemas->NumStructSchemas);
+		Plans = CreateLoadPlans(Batch, GTypes, Bindings, LoadStructIds);
 	}
 
 	~FBatchLoader()
 	{
+		CHECK(LoadIdx == Objects.Num()); // Test should load all saved objects
 		DestroyLoadPlans(Plans);
 		const FSchemaBatch* LoadSchemas = UnmountReadSchemas(Objects[0].Schema.Batch);
 		DestroyTranslatedSchemas(LoadSchemas);
 	}
 
 	template<class T>
-	T Load();
+	T Load()
+	{
+		T Out;
+		FStructView In = Objects[LoadIdx++];
+		LoadStruct(reinterpret_cast<uint8*>(&Out), In, *Plans);
+		return MoveTemp(Out);
+	}
 
 private:
-	TConstArrayView<FName> Names;
+	TConstArrayView<FName> SavedNames;
 	FLoadBatch* Plans;
 	TArray<FStructView> Objects;
 	int32 LoadIdx = 0;
-
-	//virtual FNestedScope				Resolve(FNestedScopeId Id) const override						{ return Scopes[Id.Idx]; }
-	//virtual FParametricTypeView			Resolve(FParametricTypeId Id) const override					{ return Parametrics[Id.Idx]; }
-	//virtual void						AppendDebugString(FString& Out, FNameId Name) const override	{ return Names[Name.Idx].AppendString(Out); }
 };
 
 
@@ -406,15 +409,15 @@ TEST_CASE_NAMED(FPlainPropsUeCoreTest, "System::Core::Serialization::PlainProps:
 	{
 		TScopedStructBinding<FInt, FDefaultRuntime> Int;
 
-		//TestSaveAndLoad<FDefaultRuntime>(
-		//[](FBatchSaver& Batch)
-		//{
-		//	//Batch.Save<FDefaultRuntime>(FInt{1234});
-		//}, 
-		//[](FBatchLoader& Batch)
-		//{
-		//	//CHECK(Batch.Load<FInt>().X == 1234);
-		//});
+		TestSaveAndLoad<FDefaultRuntime>(
+		[](FBatchSaver& Batch)
+		{
+			Batch.Save(FInt{1234});
+		}, 
+		[](FBatchLoader& Batch)
+		{
+			CHECK(Batch.Load<FInt>().X == 1234);
+		});
 	}
 
 	SECTION("TUniquePtr")
