PlainProps enum load/save/bind support

diff --git a/Engine/Source/Programs/PlainPropsTests/Private/PlainPropsBind.cpp b/Engine/Source/Programs/PlainPropsTests/Private/PlainPropsBind.cpp
--- a/Engine/Source/Programs/PlainPropsTests/Private/PlainPropsBind.cpp
+++ b/Engine/Source/Programs/PlainPropsTests/Private/PlainPropsBind.cpp
@@ -281,6 +281,15 @@ FIdBinding FIdTranslatorBase::TranslateIds(FMutableMemoryView To, FIdIndexerBase
 
 //////////////////////////////////////////////////////////////////////////
 
+template<class IdType>
+void RemapAll(TArrayView<IdType> Ids, FIdBinding NewIds)
+{
+	for (IdType& Id : Ids)
+	{
+		Id = NewIds.Remap(Id);
+	}
+}
+
 FSchemaBatch* CreateTranslatedSchemas(const FSchemaBatch& In, FIdBinding NewIds)
 {
 	const FMemoryView InSchemas = GetSchemaData(In);
@@ -305,10 +314,12 @@ FSchemaBatch* CreateTranslatedSchemas(const FSchemaBatch& In, FIdBinding NewIds)
 	for (FStructSchema& Schema : GetStructSchemas(*Out))
 	{
 		Schema.Type = NewIds.Remap(Schema.Type);
+		RemapAll(Schema.EditMemberNames(), NewIds);
 	}
 	for (FEnumSchema& Schema : GetEnumSchemas(*Out))
 	{
 		Schema.Type = NewIds.Remap(Schema.Type);
+		RemapAll(MakeArrayView(Schema.Footer, Schema.Num), NewIds);
 	}
 
 	return Out;
diff --git a/Engine/Source/Programs/PlainPropsTests/Private/PlainPropsInternalFormat.h b/Engine/Source/Programs/PlainPropsTests/Private/PlainPropsInternalFormat.h
--- a/Engine/Source/Programs/PlainPropsTests/Private/PlainPropsInternalFormat.h
+++ b/Engine/Source/Programs/PlainPropsTests/Private/PlainPropsInternalFormat.h
@@ -62,43 +62,32 @@ struct FStructSchema
 	uint16 NumInnerSchemas;
 	ESuper Inheritance : 2;
 	uint8 IsDense : 1;
-	// Needed? maybe sufficient check sizes match up
-	//uint8 IsContiguous : 1; // Dense and only non-bool leaves. Could extend to handle nested contiguous structs by complicating format slightly.
-
 	FMemberType Footer[0];
 	
-	static const FMemberType* GetMemberTypes(const FMemberType* Footer)
+	const FSchemaId*				GetInnerSchemas() const		{ return GetInnerSchemas(Footer, NumMembers, NumRangeTypes, NumMembers - UsesSuper(Inheritance)); }
+	FOptionalStructSchemaId			GetSuperSchema() const		{ return Inheritance != ESuper::No ? ToOptional(static_cast<FStructSchemaId>(*GetInnerSchemas())) : NoId; }
+	TConstArrayView<FMemberId>		GetMemberNames() const		{ return MakeArrayView(GetMemberNames(Footer, NumMembers, NumRangeTypes), NumMembers); }
+	TArrayView<FMemberId>			EditMemberNames() 			{ return MakeArrayView(const_cast<FMemberId*>(GetMemberNames(Footer, NumMembers, NumRangeTypes)), NumMembers); }
+
+	static const FMemberType*		GetMemberTypes(const FMemberType* Footer)
 	{
 		return Footer;
 	}
 	
-	static const FMemberType* GetRangeTypes(const FMemberType* Footer, uint32 NumMembers)
+	static const FMemberType*		GetRangeTypes(const FMemberType* Footer, uint32 NumMembers)
 	{
 		return Footer + NumMembers;
 	}
 
-	static const FMemberId*	GetMemberNames(const FMemberType* Footer, uint32 NumMembers, uint32 NumRangeTypes)
+	static const FMemberId*			GetMemberNames(const FMemberType* Footer, uint32 NumMembers, uint32 NumRangeTypes)
 	{
 		return AlignPtr<FMemberId>(Footer + NumMembers + NumRangeTypes);
 	}
-	
-	static const FSchemaId*	GetInnerSchemas(const FMemberType* Footer, uint32 NumMembers, uint32 NumRangeTypes, uint32 NumNames)
+		
+	static const FSchemaId*			GetInnerSchemas(const FMemberType* Footer, uint32 NumMembers, uint32 NumRangeTypes, uint32 NumNames)
 	{
 		return AlignPtr<FSchemaId>(GetMemberNames(Footer, NumMembers, NumRangeTypes) + NumNames);
 	}
-
-	const FSchemaId* GetInnerSchemas() const
-	{
-		return GetInnerSchemas(Footer, NumMembers, NumRangeTypes, NumMembers - UsesSuper(Inheritance));
-	}
-
-	FOptionalStructSchemaId GetSuperSchema() const
-	{
-		const FSchemaId* FirstSchema = GetInnerSchemas();
-		return Inheritance != ESuper::No ? ToOptional(static_cast<FStructSchemaId>(*FirstSchema)) : NoId;
-	}
-
-	TConstArrayView<FMemberId> GetMemberNames() const { return MakeArrayView(GetMemberNames(Footer, NumMembers, NumRangeTypes), NumMembers); }
 };
 
 static_assert(sizeof(FStructSchema) == 16 && alignof(FStructSchema) == 4,		"Add binary format versioning and read old format in AllocateReadSchemas");
@@ -164,10 +153,6 @@ inline constexpr uint64 GetLeafRangeSize(uint64 Num, FUnpackedLeafType Leaf)
 	return Leaf.Type == ELeafType::Bool ? (Num + 7) / 8 : Num * SizeOf(Leaf.Width);
 }
 
-
-
-
-
 //////////////////////////////////////////////////////////////////////////
 
 } // namespace PlainProps
\ No newline at end of file
diff --git a/Engine/Source/Programs/PlainPropsTests/Private/PlainPropsLoad.cpp b/Engine/Source/Programs/PlainPropsTests/Private/PlainPropsLoad.cpp
--- a/Engine/Source/Programs/PlainPropsTests/Private/PlainPropsLoad.cpp
+++ b/Engine/Source/Programs/PlainPropsTests/Private/PlainPropsLoad.cpp
@@ -191,10 +191,8 @@ static void CreateSubsetLoadSchema(const FStructSchema& From, const FStructSchem
 			ToIt.SkipMember();
 		}
 
-		CopyMemberBinding(/* in-out */ ToIt, Footer);			
+		CopyMemberBinding(/* in-out */ ToIt, /* out */ Footer);	
 	}
-
-	unimplemented(); // todo copy kept members but skip enum schemas
 }
 
 [[nodiscard]] static FLoadStructPlan MakeSchemaLoadPlan(const FStructSchema& From, const FStructSchemaBinding& To, TConstArrayView<FMemberId> ToMemberIds, TConstArrayView<FStructSchemaId> ToStructIds, SubsetByteArray& OutSubsetSchemas)
@@ -227,12 +225,12 @@ static void CreateSubsetLoadSchema(const FStructSchema& From, const FStructSchem
 	return Memcpy ? FLoadStructPlan(Memcpy.GetValue()) : MakeSchemaLoadPlan(From, To, ToMemberIds, ToStructIds, OutSubsetSchemas);
 }
 
-FLoadBatch* CreateLoadPlans(FReadBatchId ReadId, const FDeclarations& Declarations, const FStructBindings& Bindings, TConstArrayView<FStructSchemaId> LoadIds)
+FLoadBatch* CreateLoadPlans(FReadBatchId ReadId, const FDeclarations& Declarations, const FStructBindings& Bindings, TConstArrayView<FStructSchemaId> RuntimeIds)
 {
-	check(NumStructSchemas(ReadId) == LoadIds.Num());
+	check(NumStructSchemas(ReadId) == RuntimeIds.Num());
 
 	// Temporary data structures
-	const uint32 NumPlans = LoadIds.Num();
+	const uint32 NumPlans = RuntimeIds.Num();
 	TArray<FLoadStructPlan, TInlineAllocator<256>> Plans;
 	TArray<uint32, TInlineAllocator<256>> SubsetSchemaSizes;
 	SubsetByteArray SubsetSchemaData;
@@ -242,8 +240,8 @@ FLoadBatch* CreateLoadPlans(FReadBatchId ReadId, const FDeclarations& Declaratio
 	// Create plans
 	for (FStructSchemaId SavedId = { 0 }; SavedId.Idx < NumPlans; ++SavedId.Idx)
 	{
-		FStructSchemaId LoadId = LoadIds[SavedId.Idx];
-		FStructBinding Binding = Bindings.Get(LoadId);
+		FStructSchemaId RuntimeId = RuntimeIds[SavedId.Idx];
+		FStructBinding Binding = Bindings.Get(RuntimeId);
 		int32 SubsetSchemaOffset = SubsetSchemaData.Num();
 		if (Binding.IsCustom())
 		{
@@ -251,11 +249,11 @@ FLoadBatch* CreateLoadPlans(FReadBatchId ReadId, const FDeclarations& Declaratio
 		}
 		else
 		{
-			const FStructSchema& From = ResolveStructSchema(ReadId, LoadId);
+			const FStructSchema& From = ResolveStructSchema(ReadId, SavedId);
 			const FStructSchemaBinding& To = Binding.AsSchema();
 			// Possible optimization - some simple memcpy cases doesn't need to resolve the declaration
-			TConstArrayView<FMemberId> ToMemberIds = Declarations.Get(LoadId).GetMemberOrder();
-			Plans[SavedId.Idx] = MakeLoadPlan(From, To, ToMemberIds, LoadIds, /* out */ SubsetSchemaData);	
+			TConstArrayView<FMemberId> ToMemberIds = Declarations.Get(RuntimeId).GetMemberOrder();
+			Plans[SavedId.Idx] = MakeLoadPlan(From, To, ToMemberIds, RuntimeIds, /* out */ SubsetSchemaData);	
 		}
 		
 		SubsetSchemaSizes[SavedId.Idx] = SubsetSchemaData.Num() - SubsetSchemaOffset;
@@ -280,7 +278,7 @@ FLoadBatch* CreateLoadPlans(FReadBatchId ReadId, const FDeclarations& Declaratio
 				check(IsAligned(Size, alignof(FStructSchemaBinding)));
 				check(Plans[Idx].IsSchema());
 				bool bSparse = Plans[Idx].IsSparseSchema();
-				Plans[Idx] = FLoadStructPlan(*reinterpret_cast<const FStructSchemaBinding*>(It), ELeafWidth::B32, bSparse);
+				Out->Plans[Idx] = FLoadStructPlan(*reinterpret_cast<const FStructSchemaBinding*>(It), ELeafWidth::B32, bSparse);
 				It += Size;
 			}
 		}
diff --git a/Engine/Source/Programs/PlainPropsTests/Private/PlainPropsWrite.cpp b/Engine/Source/Programs/PlainPropsTests/Private/PlainPropsWrite.cpp
--- a/Engine/Source/Programs/PlainPropsTests/Private/PlainPropsWrite.cpp
+++ b/Engine/Source/Programs/PlainPropsTests/Private/PlainPropsWrite.cpp
@@ -27,6 +27,7 @@ struct FWriteIds
 	TArray<FOptionalSchemaId>			Enums;
 
 	uint32								NumKeptSchemas;
+	uint32								NumKeptStructSchemas;
 	TArray<FNestedScope>				KeptScopes;
 	TArray<FParametricType>				KeptParametrics;
 	TArray<FTypeId>						KeptParameters;
@@ -189,6 +190,7 @@ FWriteIds::FWriteIds(const FIdIndexerBase& Declared, const FBuiltSchemas& Schema
 	{
 		Structs.Add(Struct.bUsed ? ToOptional(FSchemaId{ NewSchemaIdx++ }) : NoId);
 	}
+	NumKeptStructSchemas = NewSchemaIdx;
 	for (const FBuiltEnumSchema& Enum : Schemas.Enums)
 	{
 		Enums.Add(Enum.bUsed ? ToOptional(FSchemaId{ NewSchemaIdx++ }) : NoId);
@@ -467,7 +469,7 @@ static void WriteSchemasImpl(TArray64<uint8>& Out, const FBuiltSchemas& Schemas,
 	Header.NestedScopesOffset = IntCastChecked<uint32>(NestedScopePos - HeaderPos);
 	Header.NumParametricTypes = NewIds.KeptParametrics.Num();
 	Header.NumSchemas = NewIds.NumKeptSchemas;
-	Header.NumStructSchemas = Schemas.Structs.Num();
+	Header.NumStructSchemas = NewIds.NumKeptStructSchemas;
 	FMemory::Memcpy(&Out[HeaderPos], &Header, sizeof(FSchemaBatch));
 
 	// Write schema offsets
diff --git a/Engine/Source/Programs/PlainPropsTests/Public/PlainPropsBind.h b/Engine/Source/Programs/PlainPropsTests/Public/PlainPropsBind.h
--- a/Engine/Source/Programs/PlainPropsTests/Public/PlainPropsBind.h
+++ b/Engine/Source/Programs/PlainPropsTests/Public/PlainPropsBind.h
@@ -616,15 +616,18 @@ template<class Ctti, class Ids>
 FEnumSchemaId DeclareNativeEnum(FDeclarations& Out, EEnumMode Mode)
 {
 	using UnderlyingType = std::underlying_type_t<typename Ctti::Type>;
-	FTypeId Type = Ids::IndexNativeType(Ctti::Name);	
+
+	FTypeId Type = Ids::IndexNativeType(Ctti::Name);
+	FEnumSchemaId Id = Ids::IndexEnum(Type);
 	FEnumerator Enumerators[Ctti::NumEnumerators];
 	for (FEnumerator& Enumerator : Enumerators)
 	{
-		Enumerator.Name = Ids::IndexMember(Ctti::Enumerators[&Enumerator - Enumerators].Name);
+		Enumerator.Name = Ids::IndexName(Ctti::Enumerators[&Enumerator - Enumerators].Name);
 		Enumerator.Constant = static_cast<uint64>(static_cast<UnderlyingType>(Ctti::Enumerators[&Enumerator - Enumerators].Constant));
 	}
+	Out.DeclareEnum(Id, Type, Mode, LeafWidth<sizeof(UnderlyingType)>, Enumerators);
 
-	return Out.DeclareEnum(Type, Mode, LeafWidth<sizeof(UnderlyingType)>, Enumerators);
+	return Id;
 }
 
 //
@@ -694,7 +697,7 @@ struct FMemberBinder
 		check(MemberIt == Schema.GetInnerRangeTypes());
 		check(Align(RangeTypeIt, alignof(uint32)) == (const void*)Schema.GetOffsets());
 		check(OffsetIt == (const void*)Schema.GetInnerSchemas());
-		check(Align(InnerSchemaIt, alignof(FRangeBinding)) == (const void*)Schema.GetRangeBindings());
+		check(Align(InnerSchemaIt, alignof(FRangeBinding)) == (const void*)Schema.GetRangeBindings() || Schema.NumInnerRanges == 0);
 		check(Schema.NumInnerRanges == RangeBindingIt - Schema.GetRangeBindings());
 	}
 
diff --git a/Engine/Source/Programs/PlainPropsTests/Public/PlainPropsCtti.h b/Engine/Source/Programs/PlainPropsTests/Public/PlainPropsCtti.h
--- a/Engine/Source/Programs/PlainPropsTests/Public/PlainPropsCtti.h
+++ b/Engine/Source/Programs/PlainPropsTests/Public/PlainPropsCtti.h
@@ -49,7 +49,7 @@ template<typename T> using CttiOf = typename TCttiOf<T>::Type;
 #define _PP_REFLECT_THINGS(N, MACRO, ...)				_PP_EXPAND(_PP_REFLECT_##N(N, MACRO, __VA_ARGS__))
 
 #define _PP_REFLECT_ENUMERATOR(N, NS, T, C)				{ #C, Type::C }, 
-#define _PP_REFLECT_ENUM_INNER(N, NS, T, ES)			struct T##_Ctti { inline static constexpr char Name[] = #T; using Type = NS :: T; static constexpr int NumEnumerators = N; static constexpr struct { const char* Name; Type Value; } Enumerators[] = { ES }; }; T##_Ctti CttiOfPtr(T*);
+#define _PP_REFLECT_ENUM_INNER(N, NS, T, ES)			struct T##_Ctti { inline static constexpr char Name[] = #T; using Type = NS :: T; static constexpr int NumEnumerators = N; static constexpr struct { const char* Name; Type Constant; } Enumerators[] = { ES }; }; T##_Ctti CttiOfPtr(T*);
 #define _PP_REFLECT_ENUM(N, NS, T, ...)					_PP_REFLECT_ENUM_INNER(N, NS, T, _PP_EXPAND(_PP_REFLECT_THINGS(N, _PP_REFLECT_ENUMERATOR, NS, T, __VA_ARGS__)))					
 #define _PP_REFLECT_STRUCT(N, NS, T, S, ...)			PP_REFLECT_STRUCT_ONLY(N, NS, T, S)				_PP_EXPAND(_PP_REFLECT_THINGS(N, PP_REFLECT_MEMBER, NS, T, __VA_ARGS__))
 #define _PP_REFLECT_STRUCT_TEMPLATE(N, NS, T, S, ...)	PP_REFLECT_STRUCT_TEMPLATE_ONLY(N, NS, T, S)	_PP_EXPAND(_PP_REFLECT_THINGS(N, PP_REFLECT_TEMPLATE_MEMBER, NS, T, __VA_ARGS__))
diff --git a/Engine/Source/Programs/PlainPropsTests/Public/PlainPropsLoad.h b/Engine/Source/Programs/PlainPropsTests/Public/PlainPropsLoad.h
--- a/Engine/Source/Programs/PlainPropsTests/Public/PlainPropsLoad.h
+++ b/Engine/Source/Programs/PlainPropsTests/Public/PlainPropsLoad.h
@@ -16,7 +16,7 @@ class FStructBindings;
 struct FTypedRange;
 
 
-PP_API FLoadBatch* CreateLoadPlans(FReadBatchId ReadId, const FDeclarations& Declarations, const FStructBindings& Bindings, TConstArrayView<FStructSchemaId> LoadIds);
+PP_API FLoadBatch* CreateLoadPlans(FReadBatchId ReadId, const FDeclarations& Declarations, const FStructBindings& Bindings, TConstArrayView<FStructSchemaId> RuntimeIds);
 PP_API void DestroyLoadPlans(FLoadBatch* Batch);
 PP_API void LoadStruct(uint8* Dst, FByteReader Src, FStructSchemaId Id, const FLoadBatch& Batch);
 PP_API void LoadStruct(uint8* Dst, FStructView Src, const FLoadBatch& Batch);
diff --git a/Engine/Source/Programs/PlainPropsTests/Tests/PlainPropsCttiTests.cpp b/Engine/Source/Programs/PlainPropsTests/Tests/PlainPropsCttiTests.cpp
--- a/Engine/Source/Programs/PlainPropsTests/Tests/PlainPropsCttiTests.cpp
+++ b/Engine/Source/Programs/PlainPropsTests/Tests/PlainPropsCttiTests.cpp
@@ -26,15 +26,15 @@ namespace Expect
 		static constexpr char Name[] = "E1";
 		using Type = ::PlainProps::Test::E1;
 		static constexpr int NumEnumerators = 2;
-		static constexpr struct { const char* Name; Type Value; } Enumerators[] = { {"A", Type::A},  {"B", Type::B} };
+		static constexpr struct { const char* Name; Type Constant; } Enumerators[] = { {"A", Type::A},  {"B", Type::B} };
 	};
 }
 
 static_assert(std::string_view(Actual::E1_Ctti::Name) == std::string_view(Expect::E1_Ctti::Name));
 static_assert(std::is_same_v<Actual::E1_Ctti::Type, Expect::E1_Ctti::Type>);
 static_assert(Actual::E1_Ctti::NumEnumerators == Expect::E1_Ctti::NumEnumerators);
-static_assert(Actual::E1_Ctti::Enumerators[0].Value	== Expect::E1_Ctti::Enumerators[0].Value);
-static_assert(Actual::E1_Ctti::Enumerators[1].Value	== Expect::E1_Ctti::Enumerators[1].Value);
+static_assert(Actual::E1_Ctti::Enumerators[0].Constant	== Expect::E1_Ctti::Enumerators[0].Constant);
+static_assert(Actual::E1_Ctti::Enumerators[1].Constant	== Expect::E1_Ctti::Enumerators[1].Constant);
 static_assert(std::string_view(Actual::E1_Ctti::Enumerators[0].Name) == std::string_view(Expect::E1_Ctti::Enumerators[0].Name));
 static_assert(std::string_view(Actual::E1_Ctti::Enumerators[1].Name) == std::string_view(Expect::E1_Ctti::Enumerators[1].Name));
 
diff --git a/Engine/Source/Programs/PlainPropsTests/Tests/PlainPropsTests.cpp b/Engine/Source/Programs/PlainPropsTests/Tests/PlainPropsTests.cpp
--- a/Engine/Source/Programs/PlainPropsTests/Tests/PlainPropsTests.cpp
+++ b/Engine/Source/Programs/PlainPropsTests/Tests/PlainPropsTests.cpp
@@ -426,7 +426,7 @@ TEST_CASE_NAMED(FPlainPropsReadWriteTest, "System::Core::Serialization::PlainPro
 			CHECK(Schema.Inheritance == ESuper::No);
 			CHECK(FStructSchema::GetMemberTypes(Schema.Footer)[0] == FUnpackedLeafType(ELeafType::Bool, ELeafWidth::B8).Pack());
 			CHECK(FStructSchema::GetMemberTypes(Schema.Footer)[8] == FUnpackedLeafType(ELeafType::Bool, ELeafWidth::B8).Pack());
-			TConstArrayView<FMemberId> MemberIds(FStructSchema::GetMemberNames(Schema.Footer, Schema.NumMembers, Schema.NumRangeTypes), Schema.NumMembers);
+			TConstArrayView<FMemberId> MemberIds = Schema.GetMemberNames();
 			CHECK(Names[MemberIds[0]] == "b1");
 			CHECK(Names[MemberIds[1]] == "b2");
 			CHECK(Names[MemberIds[2]] == "b3");
diff --git a/Engine/Source/Programs/PlainPropsTests/Tests/PlainPropsUeCoreTests.cpp b/Engine/Source/Programs/PlainPropsTests/Tests/PlainPropsUeCoreTests.cpp
--- a/Engine/Source/Programs/PlainPropsTests/Tests/PlainPropsUeCoreTests.cpp
+++ b/Engine/Source/Programs/PlainPropsTests/Tests/PlainPropsUeCoreTests.cpp
@@ -26,6 +26,7 @@ static FStructBindings		GDeltaBindings;
 
 struct FIds
 {
+	static FNameId			IndexName(FAnsiStringView Name)				{ return GNames.MakeName(FName(Name)); }
 	static FMemberId		IndexMember(FAnsiStringView Name)			{ return GNames.NameMember(FName(Name)); }
 	static FTypenameId		IndexTypename(FAnsiStringView Name)			{ return GNames.MakeTypename(FName(Name)); }
 	static FScopeId			IndexNativeScope()							{ return GNames.MakeScope(FName(UE_MODULE_NAME)); }
@@ -84,8 +85,6 @@ struct TScopedStructBinding
 	using Ctti = CttiOf<T>;
 	//constexpr bool bDeltaBindMembers = NumCustomMembers<Ctti, FDeltaRuntime::CustomBindings>() > NumCustomMembers<Ctti, FDefaultRuntime::CustomBindings>();
 
-	FStructSchemaId Id;
-
 	TScopedStructBinding(EMemberPresence Occupancy = EMemberPresence::AllowSparse)
 	: Id(DeclareNativeStruct<Ctti, Ids>(Runtime::GetDeclarations(), Occupancy))
 	{
@@ -100,12 +99,19 @@ struct TScopedStructBinding
 	{
 		Runtime::DropStruct(Id);
 	}
+
+	FStructSchemaId Id;
 };
 
-template<typename T, class Runtime>
+template<typename Enum, EEnumMode Mode, class Runtime>
 struct TScopedEnumBinding
 {
-	// todo
+	using Ids = typename Runtime::Ids;
+	using Ctti = CttiOf<Enum>;
+
+	FEnumSchemaId Id;
+	TScopedEnumBinding() : Id(DeclareNativeEnum<Ctti, Ids>(Runtime::GetDeclarations(), Mode)) {}
+	//~TScopedEnumBinding() { Runtime::GetDeclarations().DropEnum(Id)}
 };
 
 //template<typename T, class Runtime>
@@ -241,8 +247,8 @@ public:
 		CHECK(It.Grab<uint32>() == Magics[2]);
 		
 		// Bind saved ids to runtime ids, make new schemas with new ids and mount them
-		FIdTranslator LoadIds(GNames, SavedNames, *SavedSchemas);
-		FSchemaBatch* LoadSchemas = CreateTranslatedSchemas(*SavedSchemas, LoadIds.Translation);
+		FIdTranslator RuntimeIds(GNames, SavedNames, *SavedSchemas);
+		FSchemaBatch* LoadSchemas = CreateTranslatedSchemas(*SavedSchemas, RuntimeIds.Translation);
 		FReadBatchId Batch = MountReadSchemas(LoadSchemas);
 
 		// Read objects
@@ -258,7 +264,7 @@ public:
 		CHECK(!Objects.IsEmpty());
 
 		// Finally create load plans
-		TConstArrayView<FStructSchemaId> LoadStructIds = LoadIds.Translation.GetStructIds(SavedSchemas->NumStructSchemas);
+		TConstArrayView<FStructSchemaId> LoadStructIds = RuntimeIds.Translation.GetStructIds(SavedSchemas->NumStructSchemas);
 		Plans = CreateLoadPlans(Batch, GTypes, Bindings, LoadStructIds);
 	}
 
@@ -335,20 +341,6 @@ static void TestSaveAndLoad(void (*Save)(FBatchSaver&), void (*Load)(FBatchLoade
 //	return Out;
 //}
 //
-//enum EFlat { A = 1, B = 3 };
-//enum EFlag { A = 0, B = 2 };
-//
-//UE_REFLECT_ENUM(EFlat, A, B);
-//UE_REFLECT_ENUM(EFlag, A, B);
-//
-//struct FEnums
-//{
-//	EFlat Flat;
-//	EFlag Flag;
-//};
-//
-//UE_REFLECT_STRUCT(FEnums, Flat, Flag);
-//
 //struct Leaves
 //{
 //	uint16 U16;
@@ -403,31 +395,84 @@ static void TestSaveAndLoad(void (*Save)(FBatchSaver&), void (*Load)(FBatchLoade
 struct FInt { int32 X; };
 PP_REFLECT_STRUCT(PlainProps::UE::Test, FInt, void, X);
 
+
+enum class EFlat1 : uint8 { A = 1, B = 3 };
+enum class EFlat2 : uint8 { A, B };
+enum class EFlag1 : uint8 { A = 2, B = 8, AB = 10 };
+enum class EFlag2 : uint8 { A = 1, B = 2, AB = 3 };
+PP_REFLECT_ENUM(PlainProps::UE::Test, EFlat1, A, B);
+PP_REFLECT_ENUM(PlainProps::UE::Test, EFlat2, A, B);
+PP_REFLECT_ENUM(PlainProps::UE::Test, EFlag1, A, B);
+PP_REFLECT_ENUM(PlainProps::UE::Test, EFlag2, A, B);
+
+struct FEnums
+{
+	EFlat1 Flat1;
+	EFlat2 Flat2;
+	EFlag1 Flag1;
+	EFlag2 Flag2;
+
+	bool operator==(FEnums O) const { return Flat1 == O.Flat1 && Flat2 == O.Flat2 && Flag1 == O.Flag1 && Flag2 == O.Flag2; }
+};
+PP_REFLECT_STRUCT(PlainProps::UE::Test, FEnums, void, Flat1, Flat2, Flag1, Flag2);
+
+
 TEST_CASE_NAMED(FPlainPropsUeCoreTest, "System::Core::Serialization::PlainProps::UE::Core", "[Core][PlainProps][SmokeFilter]")
 {
 	SECTION("Basic")
 	{
 		TScopedStructBinding<FInt, FDefaultRuntime> Int;
-
 		TestSaveAndLoad<FDefaultRuntime>(
-		[](FBatchSaver& Batch)
-		{
-			Batch.Save(FInt{1234});
-		}, 
-		[](FBatchLoader& Batch)
-		{
-			CHECK(Batch.Load<FInt>().X == 1234);
-		});
+			[](FBatchSaver& Batch)
+			{
+				Batch.Save(FInt{1234});
+			}, 
+			[](FBatchLoader& Batch)
+			{
+				CHECK(Batch.Load<FInt>().X == 1234);
+			});
+	}
+	
+	SECTION("Enum")
+	{
+		TScopedEnumBinding<EFlat1, EEnumMode::Flat, FDefaultRuntime> Flat1;
+		TScopedEnumBinding<EFlat2, EEnumMode::Flat, FDefaultRuntime> Flat2;
+		TScopedEnumBinding<EFlag1, EEnumMode::Flag, FDefaultRuntime> Flag1;
+		TScopedEnumBinding<EFlag2, EEnumMode::Flag, FDefaultRuntime> Flag2;
+		TScopedStructBinding<FEnums, FDefaultRuntime> Int;
+		TestSaveAndLoad<FDefaultRuntime>(
+			[](FBatchSaver& Batch)
+			{
+				Batch.Save(FEnums{EFlat1::A, EFlat2::A, EFlag1::A, EFlag2::A});
+				Batch.Save(FEnums{EFlat1::A, EFlat2::A, EFlag1::B, EFlag2::B});
+				Batch.Save(FEnums{EFlat1::B, EFlat2::B, EFlag1::A, EFlag2::A});
+				Batch.Save(FEnums{EFlat1::B, EFlat2::B, EFlag1::B, EFlag2::B});
+				Batch.Save(FEnums{EFlat1::B, EFlat2::B, EFlag1::AB, EFlag2::AB});
+			}, 
+			[](FBatchLoader& Batch)
+			{
+				CHECK(Batch.Load<FEnums>() == FEnums{EFlat1::A, EFlat2::A, EFlag1::A, EFlag2::A});
+				CHECK(Batch.Load<FEnums>() == FEnums{EFlat1::A, EFlat2::A, EFlag1::B, EFlag2::B});
+				CHECK(Batch.Load<FEnums>() == FEnums{EFlat1::B, EFlat2::B, EFlag1::A, EFlag2::A});
+				CHECK(Batch.Load<FEnums>() == FEnums{EFlat1::B, EFlat2::B, EFlag1::B, EFlag2::B});
+				CHECK(Batch.Load<FEnums>() == FEnums{EFlat1::B, EFlat2::B, EFlag1::AB, EFlag2::AB});
+			});
 	}
 
-	SECTION("TUniquePtr")
+	SECTION("Sparse")
 	{}
 
 	SECTION("TArray")
 	{}
 
+	SECTION("Complex")
+	{}
+
 	SECTION("FString")
 	{}
+
+	SECTION("TUniquePtr")
+	{}
 		
 	SECTION("TSet")
 	{}
