PlainProps range load/save/bind support + basic tests for TArray bindings working

diff --git a/Engine/Source/Programs/PlainPropsTests/Private/PlainPropsBind.cpp b/Engine/Source/Programs/PlainPropsTests/Private/PlainPropsBind.cpp
--- a/Engine/Source/Programs/PlainPropsTests/Private/PlainPropsBind.cpp
+++ b/Engine/Source/Programs/PlainPropsTests/Private/PlainPropsBind.cpp
@@ -21,7 +21,7 @@ static_assert((uint8)ELeafType::Unicode		== (uint8)ELeafBindType::Unicode);
 
 uint32 FStructSchemaBinding::CalculateSize() const
 {
-	uint32 Out = sizeof(FStructSchemaBinding) + NumMembers * sizeof(FMemberBindType);
+	uint32 Out = sizeof(FStructSchemaBinding) + (NumMembers + NumInnerRanges) * sizeof(FMemberBindType);
 	Out = Align(Out + NumMembers * sizeof(uint32), sizeof(uint32));
 	Out = Align(Out + NumInnerSchemas * sizeof(FSchemaId), sizeof(FSchemaId));
 	Out = Align(Out + NumInnerRanges * sizeof(FRangeBinding), sizeof(FRangeBinding));
@@ -122,16 +122,29 @@ FRangeBinding::FRangeBinding(const IRangeBinding& Binding, ERangeSizeType SizeTy
 
 ////////////////////////////////////////////////////////////////////////////////////////////////
 
-FStructBindings::~FStructBindings()
+FStructBindingOwner::~FStructBindingOwner()
 {
-	for (FStructBinding Binding : Bindings)
+	if (Handle & FStructBinding::SchemaBit)
 	{
-		if (FStructSchemaBinding* Schema = Binding.TryGetSchema())
-		{
-			FMemory::Free(Schema);
-		}
+		FMemory::Free(Get().AsPtr());
 	}
 }
+	
+FStructBinding FStructBindingOwner::Get() const
+{
+	check(*this);
+	return FStructBinding(Handle);
+}
+
+void FStructBindingOwner::TakeOwnership(FStructBinding Binding)
+{
+	check(!*this);
+	Handle = Binding.Handle;
+	check(*this);
+}
+
+FStructBindings::~FStructBindings()
+{}
 
 void FStructBindings::BindStruct(FStructSchemaId Id, const ICustomStructBinding& Custom)
 {
@@ -189,18 +202,15 @@ void FStructBindings::Bind(FStructSchemaId Id, FStructBinding Binding)
 {
 	if (Id.Idx >= static_cast<uint32>(Bindings.Num()))
 	{
-		Bindings.SetNumZeroed(Id.Idx + 1);
+		Bindings.SetNum(Id.Idx + 1);
 	}
-	check(!Bindings[Id.Idx].IsBound());
 
-	Bindings[Id.Idx] = Binding;
+	Bindings[Id.Idx].TakeOwnership(Binding);
 }
 
 FStructBinding FStructBindings::Get(FStructSchemaId Id) const
 {
-	FStructBinding Out = Bindings[Id.Idx];
-	check(Out.IsBound());
-	return Out;
+	return Bindings[Id.Idx].Get();
 }
 
 //void FStructBindings::DropStruct(FStructSchemaId Id);
diff --git a/Engine/Source/Programs/PlainPropsTests/Private/PlainPropsLoad.cpp b/Engine/Source/Programs/PlainPropsTests/Private/PlainPropsLoad.cpp
--- a/Engine/Source/Programs/PlainPropsTests/Private/PlainPropsLoad.cpp
+++ b/Engine/Source/Programs/PlainPropsTests/Private/PlainPropsLoad.cpp
@@ -376,7 +376,7 @@ public:
 		}
 		else
 		{
-			FLoadRangeContext NoItemsCtx{.Request = {nullptr, 0}};
+			FLoadRangeContext NoItemsCtx{.Request = {Member, 0}};
 			(Binding.MakeItems)(NoItemsCtx);
 		}
 	}
diff --git a/Engine/Source/Programs/PlainPropsTests/Public/PlainPropsBind.h b/Engine/Source/Programs/PlainPropsTests/Public/PlainPropsBind.h
--- a/Engine/Source/Programs/PlainPropsTests/Public/PlainPropsBind.h
+++ b/Engine/Source/Programs/PlainPropsTests/Public/PlainPropsBind.h
@@ -237,19 +237,28 @@ public:
 	explicit FStructBinding(const FStructSchemaBinding& Schema) : Handle(uint64(&Schema) | SchemaBit) {}
 	explicit FStructBinding(const ICustomStructBinding& Custom) : Handle(uint64(&Custom)) {}
 	
-	bool						IsSchema() const	{ return Handle & SchemaBit; }
-	bool						IsCustom() const	{ return !IsSchema(); }
-	const FStructSchemaBinding&	AsSchema() const	{ check(IsSchema()); return *reinterpret_cast<FStructSchemaBinding*>(Handle & ~SchemaBit); }
-	const ICustomStructBinding&	AsCustom() const	{ check(IsCustom()); return *reinterpret_cast<ICustomStructBinding*>(Handle & ~SchemaBit); }
+	bool						IsSchema() const		{ return Handle & SchemaBit; }
+	bool						IsCustom() const		{ return !IsSchema(); }
+	const FStructSchemaBinding&	AsSchema() const		{ check(IsSchema()); return *static_cast<FStructSchemaBinding*>(AsPtr()); }
+	const ICustomStructBinding&	AsCustom() const		{ check(IsCustom()); return *static_cast<ICustomStructBinding*>(AsPtr()); }
 
 private:
 	static constexpr uint64 SchemaBit = 1;
 	uint64 Handle;
 
-	friend class FStructBindings;
-	bool						IsBound() const		{ return Handle != 0; }
-	FStructSchemaBinding*		TryGetSchema() 		{ return (IsBound() && IsSchema()) ? reinterpret_cast<FStructSchemaBinding*>(Handle & ~SchemaBit) : nullptr; }
-	FStructBinding() : Handle(0) {}
+	friend class FStructBindingOwner;
+	explicit FStructBinding(uint64 InHandle) : Handle(InHandle) {}
+	void*						AsPtr() const			{ return reinterpret_cast<void*>(Handle & ~SchemaBit); }
+};
+
+class FStructBindingOwner
+{
+	uint64 Handle = 0;
+public:
+	~FStructBindingOwner();
+	explicit					operator bool() const	{ return Handle != 0; }
+	FStructBinding				Get() const;
+	void						TakeOwnership(FStructBinding Binding);
 };
 
 ////////////////////////////////////////////////////////////////////////////////////////////////
@@ -286,13 +295,12 @@ public:
 	template<typename ItemType>
 	void Set(ItemType* Items, uint64 NumItems)
 	{
-		Set(Items, Num, sizeof(ItemType));
+		Set(Items, NumItems, sizeof(ItemType));
 	}
 
 	void Set(void* Items, uint64 NumItems, uint32 ItemSize)
 	{
-		check(NumItems > 0);
-		check(Items != Data);
+		check(NumItems == 0 || Items != Data);
 		Data = reinterpret_cast<uint8*>(Items);
 		Num = NumItems;
 		Size = ItemSize;
@@ -366,7 +374,7 @@ struct FExistingItems
 	template<typename ItemType>
 	void SetAll(const ItemType* Items, uint64 NumItems)
 	{
-		SetAll(Items, NumItems, sizeof(ItemType));
+		SetAll(FExistingItemSlice{Items, NumItems}, sizeof(ItemType));
 	}
 
 	//bool HasAll() const
@@ -430,15 +438,15 @@ public:
 	FStructBindings() = default;
 	~FStructBindings();
 
-	void						BindStruct(FStructSchemaId Id, const ICustomStructBinding& Custom);
-	void						BindStruct(FStructSchemaId Id, TConstArrayView<FMemberBinding> Schema);
-	FStructBinding				Get(FStructSchemaId Id) const;
-	void						DropStruct(FStructSchemaId Id);
+	void							BindStruct(FStructSchemaId Id, const ICustomStructBinding& Custom);
+	void							BindStruct(FStructSchemaId Id, TConstArrayView<FMemberBinding> Schema);
+	FStructBinding					Get(FStructSchemaId Id) const;
+	void							DropStruct(FStructSchemaId Id);
 
 private:
-	TArray<FStructBinding>		Bindings;
+	TArray<FStructBindingOwner>		Bindings;
 
-	void						Bind(FStructSchemaId Id, FStructBinding Binding);
+	void							Bind(FStructSchemaId Id, FStructBinding Binding);
 };
 
 ////////////////////////////////////////////////////////////////////////////////////////////////
@@ -520,8 +528,8 @@ FMemberBindType BindType(FOptionalSchemaId& OutSchema)
 {
 	if constexpr (std::is_arithmetic_v<Type> || std::is_enum_v<Type>)
 	{
-		OutSchema = std::is_enum_v<Type> ? IndexNativeEnum<Type, Ids>() : NoId;
-		return ReflectLeaf<Type>;
+		OutSchema = std::is_enum_v<Type> ? ToOptional(static_cast<FSchemaId>(IndexNativeEnum<Type, Ids>())) : NoId;
+		return FMemberBindType(ReflectLeaf<Type>);
 	}
 	else
 	{
@@ -539,44 +547,70 @@ constexpr uint32 CountRangeBindings()
 	}
 	else
 	{
-		return 1 + CountRangeBindings<typename RangeBinding::ItemType>();
+		using InnerBinding = RangeBind<typename RangeBinding::ItemType>;
+		return 1 + CountRangeBindings<InnerBinding>();
 	}
 }
 
+template<typename InnerBinding, typename InnerType>
+struct TGetInnermostImpl
+{
+	using Type = typename TGetInnermostImpl<RangeBind<InnerType>, InnerType>::Type;
+};
+
+template<typename InnerType>
+struct TGetInnermostImpl<void, InnerType>
+{
+	using Type = InnerType;
+};
+
 template<typename RangeBinding>
 struct TGetInnermost
 {
 	using InnerType = typename RangeBinding::ItemType;
-    using Type = std::conditional_t<std::is_void_v<RangeBinding>, InnerType, typename TGetInnermost<InnerType>::Type>;
+	using Type = TGetInnermostImpl<RangeBind<InnerType>, InnerType>::Type;
 };
 
 template<typename RangeBinding, uint32 N>
 TConstArrayView<FRangeBinding> GetRangeBindings()
 {
 	static_assert(N != 0);
+
+	struct FOnce
+	{
+		FOnce() : Binding(Instance, RangeSizeOf(typename RangeBinding::SizeType{})) {}
+		RangeBinding Instance;
+		FRangeBinding Binding;
+	};
+
 	if constexpr (N == 1)
 	{
-		static RangeBinding StaticInstance;
-		return {&StaticInstance, 1};
+		static FOnce Static;
+		return MakeArrayView(&Static.Binding, N);
 	}
 	else
 	{
-		struct FNestedBindings
+		using InnerType = typename RangeBinding::ItemType;
+		using InnerRangeBinding = RangeBind<InnerType>;
+
+		struct FNestedOnce : FOnce
 		{
-			FNestedBindings()
+			FNestedOnce() 
 			{
-				Instances[0] = &Instance;
-				FMemory::Memcpy(Instances + 1, GetRangeBindings<typename RangeBinding::ItemType, N - 1>().GetData(), (N - 1) * sizeof(FRangeBinding*) );
+				FMemory::Memcpy(NestedBindings, GetRangeBindings<InnerRangeBinding, N - 1>().GetData(), sizeof(NestedBindings));
 			}
-			RangeBinding Instance;
-			FRangeBinding* Instances[N];
+			
+			uint8 NestedBindings[sizeof(FRangeBinding) * (N - 1)] = {};
 		};
+		static_assert(std::is_trivially_destructible_v<FRangeBinding>);
+		static_assert(offsetof(FNestedOnce, Binding) + sizeof(FRangeBinding) == offsetof(FNestedOnce, NestedBindings));	
 
-		static FNestedBindings Static;
-		return MakeArrayView(Static.Instances);
+		static FNestedOnce Static;
+		return MakeArrayView(&Static.Binding, N);
 	}
 }
 
+
 template<class Var, class Runtime>
 FMemberBinding BindMember()
 {
diff --git a/Engine/Source/Programs/PlainPropsTests/Public/PlainPropsUeCoreBindings.h b/Engine/Source/Programs/PlainPropsTests/Public/PlainPropsUeCoreBindings.h
--- a/Engine/Source/Programs/PlainPropsTests/Public/PlainPropsUeCoreBindings.h
+++ b/Engine/Source/Programs/PlainPropsTests/Public/PlainPropsUeCoreBindings.h
@@ -107,7 +107,7 @@ struct TArrayBinding : public IRangeBinding
 
 	virtual void MakeItems(FLoadRangeContext& Ctx) const override
 	{
-		TArray<T>& Array = *Ctx.Request.GetRange<TArray<T>>();
+		TArray<T>& Array = Ctx.Request.GetRange<TArray<T>>();
 		if constexpr (std::is_default_constructible_v<T>)
 		{
 			Array.SetNum(Ctx.Request.NumTotal());
@@ -123,7 +123,7 @@ struct TArrayBinding : public IRangeBinding
 
 	virtual void ReadItems(FSaveRangeContext& Ctx) const override
 	{
-		const TArray<T>& Array = *Ctx.Request.GetRange<TArray<T>>();
+		const TArray<T>& Array = Ctx.Request.GetRange<TArray<T>>();
 		Ctx.Items.SetAll(Array.GetData(), static_cast<uint64>(Array.Num()));
 	}
 };
diff --git a/Engine/Source/Programs/PlainPropsTests/Tests/PlainPropsUeCoreTests.cpp b/Engine/Source/Programs/PlainPropsTests/Tests/PlainPropsUeCoreTests.cpp
--- a/Engine/Source/Programs/PlainPropsTests/Tests/PlainPropsUeCoreTests.cpp
+++ b/Engine/Source/Programs/PlainPropsTests/Tests/PlainPropsUeCoreTests.cpp
@@ -341,19 +341,7 @@ static void TestSaveAndLoad(void (*Save)(FBatchSaver&), void (*Load)(FBatchLoade
 //	return Out;
 //}
 //
-//struct Leaves
-//{
-//	uint16 U16;
-//	int64 I64;
-//	bool Bool;
-//	char AsciiChar;
-//	uint8 BitfieldBool : 1;
-//	uint32 StaticU32s[3];
-//};
-//
-//UE_REFLECT_STRUCT(Leaves, U16, I64, Bool, AsciiChar, BitfieldBool, StaticU32s);
 //
-//GTestBindings.Add(LeavesCtti);
 //
 //struct FNestedStructs
 //{
@@ -395,6 +383,7 @@ static void TestSaveAndLoad(void (*Save)(FBatchSaver&), void (*Load)(FBatchLoade
 struct FInt { int32 X; };
 PP_REFLECT_STRUCT(PlainProps::UE::Test, FInt, void, X);
 
+//////////////////////////////////////////////////////////////////////////
 
 enum class EFlat1 : uint8 { A = 1, B = 3 };
 enum class EFlat2 : uint8 { A, B };
@@ -416,6 +405,32 @@ struct FEnums
 };
 PP_REFLECT_STRUCT(PlainProps::UE::Test, FEnums, void, Flat1, Flat2, Flag1, Flag2);
 
+//////////////////////////////////////////////////////////////////////////
+
+//struct FLeaves
+//{
+//	uint16 U16 = 0;
+//	int64 I64 = 0;
+//	bool Bool = false;
+//	char AsciiChar = 0;
+////	uint8 BitfieldBool : 1;
+//	uint32 StaticU32s[3] = {};
+//	double Double = 0.0;
+//};
+//PP_REFLECT_STRUCT(PlainProps::UE::Test, FLeaves, void, U16, I64, Bool, AsciiChar, StaticU32s, Double);
+
+//////////////////////////////////////////////////////////////////////////
+
+struct FLeafArrays
+{
+	TArray<bool> Bits;
+	TArray<int>	Bobs;
+
+	bool operator==(const FLeafArrays& O) const { return Bits == O.Bits && Bobs == O.Bobs; }
+};
+PP_REFLECT_STRUCT(PlainProps::UE::Test, FLeafArrays, void, Bits, Bobs);
+
+//////////////////////////////////////////////////////////////////////////
 
 TEST_CASE_NAMED(FPlainPropsUeCoreTest, "System::Core::Serialization::PlainProps::UE::Core", "[Core][PlainProps][SmokeFilter]")
 {
@@ -459,13 +474,25 @@ TEST_CASE_NAMED(FPlainPropsUeCoreTest, "System::Core::Serialization::PlainProps:
 			});
 	}
 
-	SECTION("Sparse")
-	{}
-
-	SECTION("TArray")
-	{}
+	SECTION("TArrayBasic")
+	{
+		TScopedStructBinding<FLeafArrays, FDefaultRuntime> LeafArrays;
+		TestSaveAndLoad<FDefaultRuntime>(
+			[](FBatchSaver& Batch)
+			{
+				Batch.Save(FLeafArrays{{}, {}});
+				Batch.Save(FLeafArrays{{false}, {1, 2}});
+				Batch.Save(FLeafArrays{{true, false}, {3, 4, 5}});
+			}, 
+			[](FBatchLoader& Batch)
+			{
+				CHECK(Batch.Load<FLeafArrays>() == FLeafArrays{{}, {}});
+				CHECK(Batch.Load<FLeafArrays>() == FLeafArrays{{false}, {1, 2}});
+				CHECK(Batch.Load<FLeafArrays>() == FLeafArrays{{true, false}, {3, 4, 5}});
+			});
+	}
 
-	SECTION("Complex")
+	SECTION("TArrayComplex")
 	{}
 
 	SECTION("FString")
