PlainProps only build used schemas and add dropping of schemas and declarations needed for upcoming unit tests

diff --git a/Engine/Source/Programs/PlainPropsTests/Private/PlainPropsBind.cpp b/Engine/Source/Programs/PlainPropsTests/Private/PlainPropsBind.cpp
--- a/Engine/Source/Programs/PlainPropsTests/Private/PlainPropsBind.cpp
+++ b/Engine/Source/Programs/PlainPropsTests/Private/PlainPropsBind.cpp
@@ -129,13 +129,23 @@ FStructBindingOwner::~FStructBindingOwner()
 		FMemory::Free(Get().AsPtr());
 	}
 }
-	
+
 FStructBinding FStructBindingOwner::Get() const
 {
 	check(*this);
 	return FStructBinding(Handle);
 }
 
+void FStructBindingOwner::ResetOwned()
+{
+	check(*this);
+	if (Handle & FStructBinding::SchemaBit)
+	{
+		FMemory::Free(Get().AsPtr());
+	}
+	Handle = 0;
+}
+
 void FStructBindingOwner::TakeOwnership(FStructBinding Binding)
 {
 	check(!*this);
@@ -213,7 +223,10 @@ FStructBinding FStructBindings::Get(FStructSchemaId Id) const
 	return Bindings[Id.Idx].Get();
 }
 
-//void FStructBindings::DropStruct(FStructSchemaId Id);
+void FStructBindings::DropStruct(FStructSchemaId Id)
+{
+	Bindings[Id.Idx].ResetOwned();
+}
 
 
 //////////////////////////////////////////////////////////////////////////
diff --git a/Engine/Source/Programs/PlainPropsTests/Private/PlainPropsBuild.cpp b/Engine/Source/Programs/PlainPropsTests/Private/PlainPropsBuild.cpp
--- a/Engine/Source/Programs/PlainPropsTests/Private/PlainPropsBuild.cpp
+++ b/Engine/Source/Programs/PlainPropsTests/Private/PlainPropsBuild.cpp
@@ -88,6 +88,9 @@ namespace Private
 
 //////////////////////////////////////////////////////////////////////////
 
+FMemberBuilder::FMemberBuilder() {}
+FMemberBuilder::~FMemberBuilder() {}
+
 void FMemberBuilder::AddStruct(FMemberId Name, FStructSchemaId Schema, TUniquePtr<FBuiltStruct>&& Struct)
 {
 	Members.Emplace(Name, Schema, MoveTemp(Struct));
diff --git a/Engine/Source/Programs/PlainPropsTests/Private/PlainPropsBuildSchema.cpp b/Engine/Source/Programs/PlainPropsTests/Private/PlainPropsBuildSchema.cpp
--- a/Engine/Source/Programs/PlainPropsTests/Private/PlainPropsBuildSchema.cpp
+++ b/Engine/Source/Programs/PlainPropsTests/Private/PlainPropsBuildSchema.cpp
@@ -5,6 +5,7 @@
 #include  "PlainPropsInternalFormat.h"
 #include "Algo/Compare.h"
 #include "Algo/Find.h"
+#include "Containers/Map.h"
 
 namespace PlainProps
 {
@@ -43,54 +44,104 @@ static FString PrintMemberSchema(FMemberSchema Schema)
 
 //////////////////////////////////////////////////////////////////////////
 
-static void MarkInheritanceChainUsed(TArrayView<FBuiltStructSchema> OutStructs)
+struct FStructSchemaBuilder
 {
-	for (const FBuiltStructSchema& Struct : OutStructs)
-	{
-		if (Struct.bUsed)
-		{
-			for (FOptionalStructSchemaId Super = Struct.Super; Super && !OutStructs[Super.Get().Idx].bUsed; Super = OutStructs[Super.Get().Idx].Super)
-			{
-				OutStructs[Super.Get().Idx].bUsed = true;
-			}
-		}
-	}
+	const FStructDeclaration&					Declaration;
+	FSchemasBuilder&							AllSchemas;
+	const FDebugIds&							Debug;
+	TMap<FOptionalMemberId, FMemberSchema>		NotedMembers;
+	bool										bMissingMemberNoted = false;
+
+	void										NoteMembersRecursively(const FBuiltStruct& Struct);
+	void										NoteRangeRecursively(ERangeSizeType NumType, TConstArrayView<FMemberType> Types, FSchemaId InnermostSchema, const FBuiltRange& Range);
+	FBuiltStructSchema							Build() const;
+};
+
+struct FEnumSchemaBuilder
+{
+	const FEnumDeclaration&						Declaration;
+	FEnumSchemaId								Id;
+	TSet<uint64>								NotedConstants;
+
+	void										NoteValue(uint64 Value);
+	FBuiltEnumSchema							Build() const;
+};
+
+//////////////////////////////////////////////////////////////////////////
+
+FSchemasBuilder::FSchemasBuilder(FStructDeclarations InStructs, FEnumDeclarations InEnums, const FDebugIds& InDebug)
+: DeclaredStructs(InStructs)
+, DeclaredEnums(InEnums)
+, Debug(InDebug)
+{
+	StructIndices.Init(INDEX_NONE, InStructs.Num());
+	EnumIndices.Init(INDEX_NONE, InEnums.Num());
 }
 
-FSchemasBuilder::FSchemasBuilder(TConstArrayView<TUniquePtr<FStructDeclaration>> InStructs, TConstArrayView<TUniquePtr<FEnumDeclaration>> InEnums, const FDebugIds& InDebug)
-: Debug(InDebug)
+FSchemasBuilder::~FSchemasBuilder() {}
+
+template<class T, typename ...Ts>
+T& GetOrEmplace(int32& Index, TPagedArray<T, 4096>& Things, Ts&&... EmplaceArgs)
 {
-	Structs.Reserve(InStructs.Num());
-	Enums.Reserve(InEnums.Num());
-	for (const TUniquePtr<FStructDeclaration>& Declaration : InStructs)
+	if (Index == INDEX_NONE)
 	{
-		Structs.Emplace(*Declaration, *this, Debug);
-	}
-	for (const TUniquePtr<FEnumDeclaration>& Declaration : InEnums)
-	{
-		Enums.Emplace(*Declaration);
+		Index = Things.Num();
+		Things.Emplace(Forward<Ts>(EmplaceArgs)...);
 	}
+
+	return Things[Index];
 }
 
-FBuiltSchemas FSchemasBuilder::Build() const
+void FSchemasBuilder::NoteMembers(FStructSchemaId Id, const FBuiltStruct& Struct)
 {
+	checkf(!bBuilt, TEXT("Noted new members after building, built schema lack these members! Build() also returns pointers into NotedMembers which mustn't grow."));
+	FStructSchemaBuilder& Foo = GetOrEmplace(StructIndices[Id.Idx], Structs, *DeclaredStructs[Id.Idx], *this, Debug);
+	check(StructIndices[Id.Idx] != INDEX_NONE);
+	check(Foo.Declaration.Id == Id);
+	GetOrEmplace(StructIndices[Id.Idx], Structs, *DeclaredStructs[Id.Idx], *this, Debug).NoteMembersRecursively(Struct);
+}
+
+void FSchemasBuilder::NoteValue(FEnumSchemaId Id, uint64 Value)
+{
+	checkf(!bBuilt, TEXT("Noted new members after building, built schema lack these members! Build() also returns pointers into NotedMembers which mustn't grow."));
+	GetOrEmplace(EnumIndices[Id.Idx], Enums, *DeclaredEnums[Id.Idx], Id).NoteValue(Value);
+}
+
+FBuiltSchemas FSchemasBuilder::Build()
+{
+	checkf(!bBuilt, TEXT("Already built"));
+	bBuilt = true;
+
+	NoteInheritanceChains();
+
 	FBuiltSchemas Out;
 	Out.Structs.Reserve(Structs.Num());
-	for (const FStructSchemaBuilder& Struct : Structs)
+	Out.Enums.Reserve(Enums.Num());
+
+	for (FStructSchemaBuilder& Struct : Structs)
 	{
-		Out.Structs.Add(Struct.Build());
+		Out.Structs.Emplace(Struct.Build());
 	}
-	Out.Enums.Reserve(Enums.Num());
-	for (const FEnumSchemaBuilder& Enum : Enums)
+	for (FEnumSchemaBuilder& Enum : Enums)
 	{
-		Out.Enums.Add(Enum.Build());
+		Out.Enums.Emplace(Enum.Build());
 	}
 
-	MarkInheritanceChainUsed(Out.Structs);
-
 	return Out;
 }
 
+void FSchemasBuilder::NoteInheritanceChains()
+{
+	for (int Idx = 0, Num = Structs.Num(); Idx < Num; ++Idx)
+	{
+		for (FOptionalStructSchemaId Super = Structs[Idx].Declaration.Super; Super; Super = DeclaredStructs[Super.Get().Idx]->Super)
+		{
+			uint32 SuperIdx = Super.Get().Idx;
+			GetOrEmplace(StructIndices[SuperIdx], Structs, *DeclaredStructs[SuperIdx], *this, Debug);
+		}
+	}
+}
+
 //////////////////////////////////////////////////////////////////////////
 
 static bool RequiresDynamicStructSchema(const FMemberSchema& A, const FMemberSchema& B)
@@ -122,7 +173,6 @@ static void SetIsDynamic(FMemberType& InOut)
 
 void FStructSchemaBuilder::NoteMembersRecursively(const FBuiltStruct& Struct)
 {
-	checkf(!bBuilt, TEXT("Noted new members after building, built schema lack these memebers! Build() also returns pointers into NotedMembers which mustn't grow."));
 	check(Declaration.Occupancy != EMemberPresence::RequireAll || Struct.NumMembers == Declaration.NumMembers);
 	bMissingMemberNoted |= Struct.NumMembers < static_cast<uint32>(Declaration.NumMembers + !!Declaration.Super);
 	
@@ -220,13 +270,11 @@ void FStructSchemaBuilder::NoteRangeRecursively(ERangeSizeType NumType, TConstAr
 
 FBuiltStructSchema FStructSchemaBuilder::Build() const
 {
-	FBuiltStructSchema Out = { Declaration.Type };
+	FBuiltStructSchema Out = { Declaration.Type, Declaration.Id, Declaration.Super };
 	Out.bDense = Declaration.Occupancy == EMemberPresence::RequireAll || !bMissingMemberNoted;
-	Out.Super = Declaration.Super;
 	
 	if (int32 Num = NotedMembers.Num())
 	{
-		Out.bUsed = true;
 		Out.MemberNames.Reserve(Num);
 		Out.MemberSchemas.Reserve(Num);
 
@@ -253,7 +301,6 @@ FBuiltStructSchema FStructSchemaBuilder::Build() const
 		}
 	}
 
-	bBuilt = true; 
 	check(NotedMembers.Num() == Out.MemberSchemas.Num());
 	return Out;
 }
@@ -262,13 +309,12 @@ FBuiltStructSchema FStructSchemaBuilder::Build() const
 
 FBuiltEnumSchema FEnumSchemaBuilder::Build() const
 {
-	FBuiltEnumSchema Out = { Declaration.Type };
+	FBuiltEnumSchema Out = { Declaration.Type, Id };
 	Out.Mode = Declaration.Mode;
 	Out.Width = Declaration.Width;
 
 	if (int32 Num = NotedConstants.Num())
 	{
-		Out.bUsed = true;
 		Out.Names.Reserve(Num);
 		Out.Constants.Reserve(Num);
 		for (FEnumerator Enumerator : Declaration.GetEnumerators())
@@ -281,7 +327,6 @@ FBuiltEnumSchema FEnumSchemaBuilder::Build() const
 		}
 	}
 	
-	bBuilt = true; 
 	check(	NotedConstants.Num() == Out.Constants.Num() || 
 			NotedConstants.Num() == Out.Constants.Num() + (Out.Mode == EEnumMode::Flag) );
 	return Out;
diff --git a/Engine/Source/Programs/PlainPropsTests/Private/PlainPropsWrite.cpp b/Engine/Source/Programs/PlainPropsTests/Private/PlainPropsWrite.cpp
--- a/Engine/Source/Programs/PlainPropsTests/Private/PlainPropsWrite.cpp
+++ b/Engine/Source/Programs/PlainPropsTests/Private/PlainPropsWrite.cpp
@@ -76,10 +76,7 @@ struct FUsedIds
 	{
 		for (const PartialSchemaType& Schema : Schemas)
 		{
-			if (Schema.bUsed)
-			{
-				MarkUsed(Schema.Type);
-			}
+			MarkUsed(Schema.Type);
 
 			for (FNameId Name : GetUsedNames(Schema))
 			{
@@ -181,19 +178,19 @@ static void CopyUsedIds(TArray<T>& Out, const TBitArray<>& Used, C&& Ids)
 
 FWriteIds::FWriteIds(const FIdIndexerBase& Declared, const FBuiltSchemas& Schemas)
 {
-	Structs.Reserve(Schemas.Structs.Num());
-	Enums.Reserve(Schemas.Enums.Num());
+	Structs.Init(NoId, Declared.NumStructs());
+	Enums.Init(NoId, Declared.NumEnums());
 
 	// Generate new struct and enum schema indices
 	uint32 NewSchemaIdx = 0;
 	for (const FBuiltStructSchema& Struct : Schemas.Structs)
 	{
-		Structs.Add(Struct.bUsed ? ToOptional(FSchemaId{ NewSchemaIdx++ }) : NoId);
+		Structs[Struct.Id.Idx] = FSchemaId{NewSchemaIdx++};
 	}
 	NumKeptStructSchemas = NewSchemaIdx;
 	for (const FBuiltEnumSchema& Enum : Schemas.Enums)
 	{
-		Enums.Add(Enum.bUsed ? ToOptional(FSchemaId{ NewSchemaIdx++ }) : NoId);
+		Enums[Enum.Id.Idx] = FSchemaId{NewSchemaIdx++};
 	}
 	NumKeptSchemas = NewSchemaIdx;
 
@@ -434,9 +431,6 @@ void AppendBinary(TArray64<uint8>& Dst, const TArray<T>& Src)
 
 static void WriteSchemasImpl(TArray64<uint8>& Out, const FBuiltSchemas& Schemas, const FWriteIds& NewIds)
 {
-	check(NewIds.Structs.Num() == Schemas.Structs.Num());
-	check(NewIds.Enums.Num() == Schemas.Enums.Num());
-
 	// @see FReadSchemaBatch
 
 	WriteAlignmentPadding<uint32>(Out);
@@ -450,11 +444,8 @@ static void WriteSchemasImpl(TArray64<uint8>& Out, const FBuiltSchemas& Schemas,
 	{
 		for (auto& PartialSchema : PartialSchemas)
 		{
-			if (PartialSchema.bUsed)
-			{
-				SchemaOffsets.Add(IntCastChecked<uint32>(Out.Num() - HeaderPos));
-				WriteSchema(Out, PartialSchema, NewIds);
-			}
+			SchemaOffsets.Add(IntCastChecked<uint32>(Out.Num() - HeaderPos));
+			WriteSchema(Out, PartialSchema, NewIds);
 		}
 	};
 
@@ -568,7 +559,8 @@ public:
 
 	void WriteMembers(FStructSchemaId BuiltId, const FBuiltStruct& Struct)
 	{
-		const FBuiltStructSchema& Schema = Schemas.Structs[BuiltId.Idx];
+		FSchemaId WriteId = NewIds.RemapStruct(BuiltId);
+		const FBuiltStructSchema& Schema = Schemas.Structs[WriteId.Idx];
 		const TArray<FMemberId>& Order = Schema.MemberNames;
 		const int32 NumSuper = Schema.MemberSchemas.Num() - Schema.MemberNames.Num();
 		check(NumSuper == 0 || NumSuper == 1 && IsSuper(Schema.MemberSchemas[0]->Type));
diff --git a/Engine/Source/Programs/PlainPropsTests/Public/PlainPropsBind.h b/Engine/Source/Programs/PlainPropsTests/Public/PlainPropsBind.h
--- a/Engine/Source/Programs/PlainPropsTests/Public/PlainPropsBind.h
+++ b/Engine/Source/Programs/PlainPropsTests/Public/PlainPropsBind.h
@@ -255,9 +255,12 @@ class FStructBindingOwner
 {
 	uint64 Handle = 0;
 public:
+	UE_NONCOPYABLE(FStructBindingOwner);
+	FStructBindingOwner() = default;
 	~FStructBindingOwner();
 	explicit					operator bool() const	{ return Handle != 0; }
 	FStructBinding				Get() const;
+	void						ResetOwned();
 	void						TakeOwnership(FStructBinding Binding);
 };
 
@@ -496,12 +499,13 @@ FStructSchemaId BindCustomStructOnce()
 		FBinding()
 		: Id(Runtime::GetTypes().DeclareStruct(Ids::IndexNativeType(Type::Name), NoId, CustomBinding::GetMemberIds(), CustomBinding::Occupancy))
 		{
-			Runtime::GetCustomBindings().BindStruct(Id, Instance);
+			Runtime::GetBindings().BindStruct(Id, Instance);
 		}
 
 		~FBinding()
 		{
-			Runtime::DropStruct(Id);
+			Runtime::GetBindings().DropStruct(Id);
+			Runtime::GetDeclarations().DropStruct(Id);
 		}
 
 		CustomBinding Instance;
diff --git a/Engine/Source/Programs/PlainPropsTests/Public/PlainPropsBuild.h b/Engine/Source/Programs/PlainPropsTests/Public/PlainPropsBuild.h
--- a/Engine/Source/Programs/PlainPropsTests/Public/PlainPropsBuild.h
+++ b/Engine/Source/Programs/PlainPropsTests/Public/PlainPropsBuild.h
@@ -146,6 +146,9 @@ template<typename EnumType, typename SizeType>
 class FMemberBuilder
 {
 public:
+	FMemberBuilder();
+	~FMemberBuilder();
+
 	template<typename LeafType>
 	void Add(FMemberId Name, LeafType Value)
 	{
diff --git a/Engine/Source/Programs/PlainPropsTests/Public/PlainPropsBuildSchema.h b/Engine/Source/Programs/PlainPropsTests/Public/PlainPropsBuildSchema.h
--- a/Engine/Source/Programs/PlainPropsTests/Public/PlainPropsBuildSchema.h
+++ b/Engine/Source/Programs/PlainPropsTests/Public/PlainPropsBuildSchema.h
@@ -5,86 +5,65 @@
 #include  "PlainPropsBuild.h"
 #include "Containers/Array.h"
 #include "Containers/ArrayView.h"
-#include "Containers/Map.h"
+#include "Containers/PagedArray.h"
 
 namespace PlainProps
 {
 
 struct FMemberSchema;
+struct FEnumSchemaBuilder;
+struct FStructSchemaBuilder;
 
 struct FBuiltStructSchema
 {
-	FTypeId Type;
-	bool bUsed = false;
-	bool bDense = false;
-	FOptionalStructSchemaId Super;
-	TArray<FMemberId> MemberNames;
-	TArray<const FMemberSchema*> MemberSchemas;
+	FTypeId							Type;
+	FStructSchemaId					Id;
+	FOptionalStructSchemaId			Super;
+	bool							bDense = false;
+	TArray<FMemberId>				MemberNames;
+	TArray<const FMemberSchema*>	MemberSchemas;
 };
 
 struct FBuiltEnumSchema
 {
-	FTypeId Type;
-	bool bUsed = false;
-	EEnumMode Mode = EEnumMode::Flat;
-	ELeafWidth Width = ELeafWidth::B8;
-	TArray<FNameId> Names;
-	TArray<uint64> Constants;
+	FTypeId							Type;
+	FEnumSchemaId					Id;
+	EEnumMode						Mode = EEnumMode::Flat;
+	ELeafWidth						Width = ELeafWidth::B8;
+	TArray<FNameId>					Names;
+	TArray<uint64>					Constants;
 };
 
 struct FBuiltSchemas
 {
-	TArray<FBuiltStructSchema> Structs; // Same size as number of declared structs
-	TArray<FBuiltEnumSchema> Enums; // Same size as number of declared enums
-};
-
-class FSchemasBuilder;
-
-struct FStructSchemaBuilder
-{
-	const FStructDeclaration&		Declaration;
-	FSchemasBuilder&				AllSchemas;
-	const FDebugIds&				Debug;
-	TMap<FOptionalMemberId, FMemberSchema>	NotedMembers;
-	bool							bMissingMemberNoted = false;
-	mutable bool					bBuilt = false;
-
-	void							NoteMembersRecursively(const FBuiltStruct& Struct);
-	void							NoteRangeRecursively(ERangeSizeType NumType, TConstArrayView<FMemberType> Types, FSchemaId InnermostSchema, const FBuiltRange& Range);
-	FBuiltStructSchema				Build() const;
-};
-
-struct FEnumSchemaBuilder
-{
-	const FEnumDeclaration&			Declaration;
-	TSet<uint64>					NotedConstants;
-	mutable bool					bBuilt = false;
-
-	void							NoteValue(uint64 Value);
-	FBuiltEnumSchema				Build() const;
+	TArray<FBuiltStructSchema>		Structs; // Same size as number of declared structs
+	TArray<FBuiltEnumSchema>		Enums; // Same size as number of declared enums
 };
 
 class FSchemasBuilder
 {
 public:
-	FSchemasBuilder(TConstArrayView<TUniquePtr<FStructDeclaration>> InStructs, TConstArrayView<TUniquePtr<FEnumDeclaration>> InEnums, const FDebugIds& InDebug);
+	using FStructDeclarations = TConstArrayView<TUniquePtr<FStructDeclaration>>;
+	using FEnumDeclarations = TConstArrayView<TUniquePtr<FEnumDeclaration>>;
 
-	void NoteMembers(FStructSchemaId Schema, const FBuiltStruct& Struct)
-	{
-		Structs[Schema.Idx].NoteMembersRecursively(Struct);
-	}
+	FSchemasBuilder(FStructDeclarations InStructs, FEnumDeclarations InEnums, const FDebugIds& InDebug);
+	~FSchemasBuilder();
 
-	void NoteValue(FEnumSchemaId Schema, uint64 Value)
-	{
-		Enums[Schema.Idx].NoteValue(Value);
-	}
-	
-	FBuiltSchemas Build() const;
+	void										NoteMembers(FStructSchemaId Id, const FBuiltStruct& Struct);
+	void										NoteValue(FEnumSchemaId Id, uint64 Value);
+	FBuiltSchemas								Build();
 
 private:
-	TArray<FStructSchemaBuilder> Structs;
-	TArray<FEnumSchemaBuilder> Enums;
-	const FDebugIds& Debug;
+	FStructDeclarations							DeclaredStructs;
+	FEnumDeclarations							DeclaredEnums;
+	TArray<int32>								StructIndices;
+	TArray<int32>								EnumIndices;
+	TPagedArray<FStructSchemaBuilder, 4096>		Structs;		// TPagedArray for stable references
+	TPagedArray<FEnumSchemaBuilder, 4096>		Enums;			// TPagedArray for stable references
+	const FDebugIds&							Debug;
+	bool										bBuilt = false;
+
+	void										NoteInheritanceChains();
 };
 
 } // namespace PlainProps
\ No newline at end of file
diff --git a/Engine/Source/Programs/PlainPropsTests/Public/PlainPropsDeclare.h b/Engine/Source/Programs/PlainPropsTests/Public/PlainPropsDeclare.h
--- a/Engine/Source/Programs/PlainPropsTests/Public/PlainPropsDeclare.h
+++ b/Engine/Source/Programs/PlainPropsTests/Public/PlainPropsDeclare.h
@@ -54,8 +54,11 @@ public:
 	void											DeclareEnum(FEnumSchemaId Id, FTypeId Type, EEnumMode Mode, ELeafWidth Width, TConstArrayView<FEnumerator> Enumerators);
 	void											DeclareStruct(FStructSchemaId Id, FTypeId Type, TConstArrayView<FMemberId> MemberOrder, EMemberPresence Occupancy, FOptionalStructSchemaId Super = {});
 	
-	const FEnumDeclaration&							Get(FEnumSchemaId Id) const		{ return *DeclaredEnums[Id.Idx]; }
-	const FStructDeclaration&						Get(FStructSchemaId Id) const	{ return *DeclaredStructs[Id.Idx]; }
+	void											DropEnum(FEnumSchemaId Id)		{ Get(Id); DeclaredEnums[Id.Idx].Reset(); }
+	void											DropStruct(FStructSchemaId Id)	{ Get(Id); DeclaredStructs[Id.Idx].Reset(); }
+
+	const FEnumDeclaration&							Get(FEnumSchemaId Id) const		{ check(DeclaredEnums[Id.Idx]);		return *DeclaredEnums[Id.Idx]; }
+	const FStructDeclaration&						Get(FStructSchemaId Id) const	{ check(DeclaredStructs[Id.Idx]);	return *DeclaredStructs[Id.Idx]; }
 	
 	TConstArrayView<TUniquePtr<FEnumDeclaration>>	GetEnums() const	{ return DeclaredEnums; }
 	TConstArrayView<TUniquePtr<FStructDeclaration>>	GetStructs() const	{ return DeclaredStructs; }
@@ -63,7 +66,6 @@ public:
 protected:
 	TArray<TUniquePtr<FEnumDeclaration>>			DeclaredEnums;
 	TArray<TUniquePtr<FStructDeclaration>>			DeclaredStructs;
-
 };
 
 } // namespace PlainProps
\ No newline at end of file
diff --git a/Engine/Source/Programs/PlainPropsTests/Public/PlainPropsIndex.h b/Engine/Source/Programs/PlainPropsTests/Public/PlainPropsIndex.h
--- a/Engine/Source/Programs/PlainPropsTests/Public/PlainPropsIndex.h
+++ b/Engine/Source/Programs/PlainPropsTests/Public/PlainPropsIndex.h
@@ -63,6 +63,8 @@ public:
 	const FParametricTypeIndexer&		GetParametricTypes() const							{ return ParametricTypes; }
 
 	virtual uint32						NumNames() const = 0;
+	uint32								NumEnums() const									{ return Enums.Num(); }
+	uint32								NumStructs() const									{ return Structs.Num(); }
 
 protected:
 	FNestedScopeIndexer					NestedScopes;
diff --git a/Engine/Source/Programs/PlainPropsTests/Public/PlainPropsWrite.h b/Engine/Source/Programs/PlainPropsTests/Public/PlainPropsWrite.h
--- a/Engine/Source/Programs/PlainPropsTests/Public/PlainPropsWrite.h
+++ b/Engine/Source/Programs/PlainPropsTests/Public/PlainPropsWrite.h
@@ -2,8 +2,8 @@
 
 #pragma once
 
-#include "Containers/Array.h"
 #include  "PlainPropsTypes.h"
+#include "Containers/Array.h"
 #include "Templates/UniquePtr.h"
 
 namespace PlainProps
@@ -16,26 +16,6 @@ class FNestedScopeIndexer;
 class FParametricTypeIndexer;
 struct FWriteIds;
 
-///// Maps declared / built ids to write ids, the subset that's used by the partial schemas
-//class FUsedIds
-//{
-//public:
-//	virtual ~FUsedIds() {}
-//	virtual bool Uses(FNameId BuiltId) const = 0;
-//	virtual FOptionalStructSchemaId GetWriteId(FStructSchemaId BuiltId) const = 0;
-//};
-//
-//PP_API TUniquePtr<FUsedIds>	DropUnusedIds(const FBuiltSchemas& AllSchemas, int32 NumDeclaredNames);
-//PP_API void					WriteSchemas(TArray64<uint8>& Out, const FBuiltSchemas& AllSchemas, const FUsedIds& WriteIds);
-//PP_API void					WriteMembers(TArray64<uint8>& Out, const FBuiltSchemas& AllSchemas, const FUsedIds& WriteIds, FStructSchemaId Schema, const FBuiltStruct& Struct, const FDebugIds& Debug);
-
-//struct FDeclaredIds
-//{
-//	int32 NumNames;
-//	const FNestedScopeIndexer& NestedScopes;
-//	const FParametricTypeIndexer& ParametricTypes;
-//};
-
 class FWriter
 {
 public:
@@ -43,7 +23,7 @@ public:
 	PP_API ~FWriter();
 	
 	PP_API bool								Uses(FNameId BuiltId) const;
-	PP_API TConstArrayView<FNestedScope>		GetUsedScopes() const;
+	PP_API TConstArrayView<FNestedScope>	GetUsedScopes() const;
 	PP_API TConstArrayView<FParametricType>	GetUsedParametrics() const;
 	PP_API FOptionalStructSchemaId			GetWriteId(FStructSchemaId BuiltId) const;
 
@@ -51,9 +31,9 @@ public:
 	PP_API void								WriteMembers(TArray64<uint8>& Out, FStructSchemaId BuiltId, const FBuiltStruct& Struct) const;
 
 private:
-	const FBuiltSchemas& Schemas;
-	const FDebugIds& Debug;
-	TUniquePtr<FWriteIds> NewIds;
+	const FBuiltSchemas&					Schemas;
+	const FDebugIds&						Debug;
+	TUniquePtr<FWriteIds>					NewIds;
 };
 
 } // namespace PlainProps
\ No newline at end of file
diff --git a/Engine/Source/Programs/PlainPropsTests/Tests/PlainPropsUeCoreTests.cpp b/Engine/Source/Programs/PlainPropsTests/Tests/PlainPropsUeCoreTests.cpp
--- a/Engine/Source/Programs/PlainPropsTests/Tests/PlainPropsUeCoreTests.cpp
+++ b/Engine/Source/Programs/PlainPropsTests/Tests/PlainPropsUeCoreTests.cpp
@@ -2,8 +2,6 @@
 
 #if WITH_TESTS
 
-#include "Tests/TestHarnessAdapter.h"
-#include "Containers/StringView.h"
 #include  "PlainPropsBuildSchema.h"
 #include  "PlainPropsCtti.h"
 #include  "PlainPropsIndex.h"
@@ -14,6 +12,9 @@
 #include  "PlainPropsSave.h"
 #include  "PlainPropsWrite.h"
 #include  "PlainPropsUeCoreBindings.h"
+#include "Tests/TestHarnessAdapter.h"
+#include "Containers/StringView.h"
+#include "Containers/Map.h"
 #include "Templates/UnrealTemplate.h"
 
 namespace PlainProps::UE::Test
@@ -65,7 +66,6 @@ struct FDefaultRuntime
 
 	static FDeclarations&			GetDeclarations()				{ return GTypes; }
 	static FStructBindings&			GetBindings()					{ return GBindings; }
-	static void						DropStruct(FStructSchemaId Id)	{} // todo
 };
 
 struct FDeltaRuntime : FDefaultRuntime
@@ -73,7 +73,6 @@ struct FDeltaRuntime : FDefaultRuntime
 	template<class T> using CustomBindings = TCustomDeltaBindings<T>;
 
 	static FStructBindings&			GetBindings()					{ return GDeltaBindings; }
-	static void						DropStruct(FStructSchemaId Id)	{} // todo,  drop from default too
 };
 
 ////////////////////////////////////////////////////////////////////////////////////////////////
@@ -97,7 +96,8 @@ struct TScopedStructBinding
 
 	~TScopedStructBinding()
 	{
-		Runtime::DropStruct(Id);
+		Runtime::GetBindings().DropStruct(Id);
+		Runtime::GetDeclarations().DropStruct(Id);
 	}
 
 	FStructSchemaId Id;
@@ -111,7 +111,7 @@ struct TScopedEnumBinding
 
 	FEnumSchemaId Id;
 	TScopedEnumBinding() : Id(DeclareNativeEnum<Ctti, Ids>(Runtime::GetDeclarations(), Mode)) {}
-	//~TScopedEnumBinding() { Runtime::GetDeclarations().DropEnum(Id)}
+	~TScopedEnumBinding() { Runtime::GetDeclarations().DropEnum(Id); }
 };
 
 //template<typename T, class Runtime>
@@ -430,6 +430,17 @@ struct FLeafArrays
 };
 PP_REFLECT_STRUCT(PlainProps::UE::Test, FLeafArrays, void, Bits, Bobs);
 
+struct FComplexArrays
+{
+	TArray<char> Str;
+	TArray<EFlat1> Enums;
+	TArray<FLeafArrays> Misc;
+	TArray<TArray<EFlat1>> Nested;
+
+	bool operator==(const FComplexArrays& O) const { return Str == O.Str && Enums == O.Enums && Misc == O.Misc && Nested == O.Nested; }
+};
+PP_REFLECT_STRUCT(PlainProps::UE::Test, FComplexArrays, void, Str, Enums, Misc, Nested);
+
 //////////////////////////////////////////////////////////////////////////
 
 TEST_CASE_NAMED(FPlainPropsUeCoreTest, "System::Core::Serialization::PlainProps::UE::Core", "[Core][PlainProps][SmokeFilter]")
