PlainProps note/load/bind for nested/enum/struct ranges

diff --git a/Engine/Source/Programs/PlainPropsTests/Private/PlainPropsBind.cpp b/Engine/Source/Programs/PlainPropsTests/Private/PlainPropsBind.cpp
--- a/Engine/Source/Programs/PlainPropsTests/Private/PlainPropsBind.cpp
+++ b/Engine/Source/Programs/PlainPropsTests/Private/PlainPropsBind.cpp
@@ -201,7 +201,10 @@ void FStructBindings::BindStruct(FStructSchemaId Id, TConstArrayView<FMemberBind
 			Footer.AddRange(Ranges, Member.InnermostType, IntCastChecked<uint32>(Member.Offset));
 		}
 
-		Footer.AddOptionalInnerSchema(Member.InnermostSchema);
+		if (Member.InnermostSchema)
+		{
+			Footer.AddInnerSchema(Member.InnermostSchema.Get());
+		}
 	}
 
 	// Register
diff --git a/Engine/Source/Programs/PlainPropsTests/Private/PlainPropsBuildSchema.cpp b/Engine/Source/Programs/PlainPropsTests/Private/PlainPropsBuildSchema.cpp
--- a/Engine/Source/Programs/PlainPropsTests/Private/PlainPropsBuildSchema.cpp
+++ b/Engine/Source/Programs/PlainPropsTests/Private/PlainPropsBuildSchema.cpp
@@ -53,7 +53,7 @@ struct FStructSchemaBuilder
 	bool										bMissingMemberNoted = false;
 
 	void										NoteMembersRecursively(const FBuiltStruct& Struct);
-	void										NoteRangeRecursively(ERangeSizeType NumType, TConstArrayView<FMemberType> Types, FSchemaId InnermostSchema, const FBuiltRange& Range);
+	void										NoteRangeRecursively(ERangeSizeType NumType, TConstArrayView<FMemberType> Types, void* InnermostSchemaBuilder, const FBuiltRange& Range);
 	FBuiltStructSchema							Build() const;
 };
 
@@ -91,20 +91,24 @@ T& GetOrEmplace(int32& Index, TPagedArray<T, 4096>& Things, Ts&&... EmplaceArgs)
 
 	return Things[Index];
 }
+	
+FEnumSchemaBuilder&	FSchemasBuilder::NoteEnum(FEnumSchemaId Id)
+{
+	checkf(!bBuilt, TEXT("Noted new members after building"));
+	checkf(DeclaredEnums[Id.Idx], TEXT("Undeclared enum '%s' noted"), *Debug.Print(Id));
+	return GetOrEmplace(EnumIndices[Id.Idx], Enums, *DeclaredEnums[Id.Idx], Id);
+}
 
-void FSchemasBuilder::NoteMembers(FStructSchemaId Id, const FBuiltStruct& Struct)
+FStructSchemaBuilder& FSchemasBuilder::NoteStruct(FStructSchemaId Id)
 {
-	checkf(!bBuilt, TEXT("Noted new members after building, built schema lack these members! Build() also returns pointers into NotedMembers which mustn't grow."));
-	FStructSchemaBuilder& Foo = GetOrEmplace(StructIndices[Id.Idx], Structs, *DeclaredStructs[Id.Idx], *this, Debug);
-	check(StructIndices[Id.Idx] != INDEX_NONE);
-	check(Foo.Declaration.Id == Id);
-	GetOrEmplace(StructIndices[Id.Idx], Structs, *DeclaredStructs[Id.Idx], *this, Debug).NoteMembersRecursively(Struct);
+	checkf(!bBuilt, TEXT("Noted new members after building"));
+	checkf(DeclaredStructs[Id.Idx], TEXT("Undeclared struct '%s' noted"), *Debug.Print(Id));
+	return GetOrEmplace(StructIndices[Id.Idx], Structs, *DeclaredStructs[Id.Idx], *this, Debug);
 }
 
-void FSchemasBuilder::NoteValue(FEnumSchemaId Id, uint64 Value)
+void FSchemasBuilder::NoteStructAndMembers(FStructSchemaId Id, const FBuiltStruct& Struct)
 {
-	checkf(!bBuilt, TEXT("Noted new members after building, built schema lack these members! Build() also returns pointers into NotedMembers which mustn't grow."));
-	GetOrEmplace(EnumIndices[Id.Idx], Enums, *DeclaredEnums[Id.Idx], Id).NoteValue(Value);
+	NoteStruct(Id).NoteMembersRecursively(Struct);
 }
 
 FBuiltSchemas FSchemasBuilder::Build()
@@ -161,7 +165,6 @@ static bool RequiresDynamicStructSchema(const FMemberSchema& A, const FMemberSch
 	}
 
 	return false;
-		
 }
 
 static void SetIsDynamic(FMemberType& InOut)
@@ -171,6 +174,11 @@ static void SetIsDynamic(FMemberType& InOut)
 	InOut = FMemberType(Type);
 }
 
+static void* NoteStructOrEnum(FSchemasBuilder& AllSchemas, bool bStruct, FSchemaId Id)
+{
+	return bStruct ? static_cast<void*>(&AllSchemas.NoteStruct(static_cast<FStructSchemaId>(Id))) : &AllSchemas.NoteEnum(static_cast<FEnumSchemaId>(Id));
+}
+
 void FStructSchemaBuilder::NoteMembersRecursively(const FBuiltStruct& Struct)
 {
 	check(Declaration.Occupancy != EMemberPresence::RequireAll || Struct.NumMembers == Declaration.NumMembers);
@@ -201,16 +209,25 @@ void FStructSchemaBuilder::NoteMembersRecursively(const FBuiltStruct& Struct)
 		{
 			if (ELeafType::Enum == Schema.Type.AsLeaf().Type)
 			{
-				AllSchemas.NoteValue(static_cast<FEnumSchemaId>(Schema.InnerSchema.Get()), Member.Value.Leaf);
+				AllSchemas.NoteEnum(static_cast<FEnumSchemaId>(Schema.InnerSchema.Get())).NoteValue(Member.Value.Leaf);
 			}
 		}
 		else if (Kind == EMemberKind::Struct)
 		{
-			AllSchemas.NoteMembers(static_cast<FStructSchemaId>(Schema.InnerSchema.Get()), *Member.Value.Struct);
+			AllSchemas.NoteStruct(static_cast<FStructSchemaId>(Schema.InnerSchema.Get())).NoteMembersRecursively(*Member.Value.Struct);
 		}
-		else if (Member.Value.Range && IsStructOrEnum(MakeArrayView(Schema.InnerRangeTypes).Last()))
+		else 
 		{
-			NoteRangeRecursively(Schema.Type.AsRange().MaxSize, Schema.InnerRangeTypes, Schema.InnerSchema.Get(), *Member.Value.Range);
+			check(IsStructOrEnum(Schema.InnerRangeTypes.Last()) == !!Schema.InnerSchema);
+
+			if (Schema.InnerSchema)
+			{
+				void* InnerSchemaBuilder = NoteStructOrEnum(AllSchemas, Schema.InnerRangeTypes.Last().IsStruct(), Schema.InnerSchema.Get());
+				if (Member.Value.Range)
+				{
+					NoteRangeRecursively(Schema.Type.AsRange().MaxSize, Schema.InnerRangeTypes, InnerSchemaBuilder, *Member.Value.Range);			
+				}
+			}
 		}
 	}
 }
@@ -222,27 +239,27 @@ TConstArrayView64<T> MakeArrayView(const T* Data, uint64 Num)
 }
 
 template<typename IntType>
-void NoteEnumValues(FSchemasBuilder& AllSchemas, FEnumSchemaId Enum, const IntType* Values, uint64 Num)
+void NoteEnumValues(FEnumSchemaBuilder& Schema, const IntType* Values, uint64 Num)
 {
 	for (IntType Value : TConstArrayView64<IntType>(Values, Num))
 	{
-		AllSchemas.NoteValue(Enum, Value);
+		Schema.NoteValue(Value);
 	}
 }
 
-static void NoteEnumRange(FSchemasBuilder& Out, FLeafType Leaf, FEnumSchemaId Enum, const FBuiltRange& Range)
+static void NoteEnumRange(FEnumSchemaBuilder& Out, FLeafType Leaf, const FBuiltRange& Range)
 {
 	check(Leaf.Type == ELeafType::Enum);
 	switch (Leaf.Width)
 	{
-	case ELeafWidth::B8:	NoteEnumValues(Out, Enum, reinterpret_cast<const uint8* >(Range.Data), Range.Num); break;
-	case ELeafWidth::B16:	NoteEnumValues(Out, Enum, reinterpret_cast<const uint16*>(Range.Data), Range.Num); break;
-	case ELeafWidth::B32:	NoteEnumValues(Out, Enum, reinterpret_cast<const uint32*>(Range.Data), Range.Num); break;
-	case ELeafWidth::B64:	NoteEnumValues(Out, Enum, reinterpret_cast<const uint64*>(Range.Data), Range.Num); break;
+	case ELeafWidth::B8:	NoteEnumValues(Out, reinterpret_cast<const uint8* >(Range.Data), Range.Num); break;
+	case ELeafWidth::B16:	NoteEnumValues(Out, reinterpret_cast<const uint16*>(Range.Data), Range.Num); break;
+	case ELeafWidth::B32:	NoteEnumValues(Out, reinterpret_cast<const uint32*>(Range.Data), Range.Num); break;
+	case ELeafWidth::B64:	NoteEnumValues(Out, reinterpret_cast<const uint64*>(Range.Data), Range.Num); break;
 	}
 }
 
-void FStructSchemaBuilder::NoteRangeRecursively(ERangeSizeType NumType, TConstArrayView<FMemberType> Types, FSchemaId InnermostSchema, const FBuiltRange& Range)
+void FStructSchemaBuilder::NoteRangeRecursively(ERangeSizeType NumType, TConstArrayView<FMemberType> Types, void* InnermostSchema, const FBuiltRange& Range)
 {
 	FMemberType Type = Types[0];
 	switch (Type.GetKind())
@@ -250,7 +267,7 @@ void FStructSchemaBuilder::NoteRangeRecursively(ERangeSizeType NumType, TConstAr
 	case EMemberKind::Struct:
 		for (const TUniquePtr<const FBuiltStruct>& Struct : Range.AsStructs())
 		{
-			AllSchemas.NoteMembers(static_cast<FStructSchemaId>(InnermostSchema), *Struct);
+			static_cast<FStructSchemaBuilder*>(InnermostSchema)->NoteMembersRecursively(*Struct);
 		}
 		break;
 	case EMemberKind::Range:
@@ -263,7 +280,7 @@ void FStructSchemaBuilder::NoteRangeRecursively(ERangeSizeType NumType, TConstAr
 		}
 		break;
 	case EMemberKind::Leaf:
-		NoteEnumRange(/* out */ AllSchemas, Type.AsLeaf(), static_cast<FEnumSchemaId>(InnermostSchema), Range);
+		NoteEnumRange(/* out */ *static_cast<FEnumSchemaBuilder*>(InnermostSchema), Type.AsLeaf(), Range);
 		break;
 	}
 }
diff --git a/Engine/Source/Programs/PlainPropsTests/Private/PlainPropsInternalFormat.h b/Engine/Source/Programs/PlainPropsTests/Private/PlainPropsInternalFormat.h
--- a/Engine/Source/Programs/PlainPropsTests/Private/PlainPropsInternalFormat.h
+++ b/Engine/Source/Programs/PlainPropsTests/Private/PlainPropsInternalFormat.h
@@ -64,11 +64,13 @@ struct FStructSchema
 	uint8 IsDense : 1;
 	FMemberType Footer[0];
 	
-	const FSchemaId*				GetInnerSchemas() const		{ return GetInnerSchemas(Footer, NumMembers, NumRangeTypes, NumMembers - UsesSuper(Inheritance)); }
-	FOptionalStructSchemaId			GetSuperSchema() const		{ return Inheritance != ESuper::No ? ToOptional(static_cast<FStructSchemaId>(*GetInnerSchemas())) : NoId; }
+	TConstArrayView<FMemberType>	GetMemberTypes() const		{ return MakeArrayView(GetMemberTypes(Footer), NumMembers); }
+	TConstArrayView<FMemberType>	GetRangeTypes() const		{ return MakeArrayView(GetRangeTypes(Footer, NumMembers), NumRangeTypes); }
 	TConstArrayView<FMemberId>		GetMemberNames() const		{ return MakeArrayView(GetMemberNames(Footer, NumMembers, NumRangeTypes), NumMembers); }
 	TArrayView<FMemberId>			EditMemberNames() 			{ return MakeArrayView(const_cast<FMemberId*>(GetMemberNames(Footer, NumMembers, NumRangeTypes)), NumMembers); }
-
+	const FSchemaId*				GetInnerSchemas() const		{ return GetInnerSchemas(Footer, NumMembers, NumRangeTypes, NumMembers - UsesSuper(Inheritance)); }
+	FOptionalStructSchemaId			GetSuperSchema() const		{ return Inheritance != ESuper::No ? ToOptional(static_cast<FStructSchemaId>(*GetInnerSchemas())) : NoId; }
+	
 	static const FMemberType*		GetMemberTypes(const FMemberType* Footer)
 	{
 		return Footer;
diff --git a/Engine/Source/Programs/PlainPropsTests/Private/PlainPropsLoad.cpp b/Engine/Source/Programs/PlainPropsTests/Private/PlainPropsLoad.cpp
--- a/Engine/Source/Programs/PlainPropsTests/Private/PlainPropsLoad.cpp
+++ b/Engine/Source/Programs/PlainPropsTests/Private/PlainPropsLoad.cpp
@@ -85,31 +85,36 @@ private:
 
 static uint16 CountEnums(const FStructSchema& Schema)
 {
-	uint16 Num = 0;
-	for (FMemberType Member : MakeArrayView(Schema.Footer, Schema.NumMembers))
+	if (Schema.NumInnerSchemas == 0)
 	{
-		Num += IsEnum(Member);
+		return 0;
 	}
-	return Num;
-}
-
-static uint16 CountEnums(const FStructSchemaBinding& Schema)
-{
+	
 	uint16 Num = 0;
-	for (FMemberBindType Member : MakeArrayView(Schema.Members, Schema.NumMembers))
+	TConstArrayView<FMemberType> RangeTypes = Schema.GetRangeTypes(); 
+	if (RangeTypes.IsEmpty())
 	{
-		Num += (Member.IsLeaf() && Member.AsLeaf().Bind.Type == ELeafBindType::Enum);
+		for (FMemberType Member : Schema.GetMemberTypes())
+		{
+			Num += IsEnum(Member);
+		}
+		return Num;
 	}
-	return Num;
-}
-
-static uint16 CountStaticStructs(const FStructSchemaBinding& Schema)
-{
-	uint16 Num = 0;
-	for (FMemberBindType Member : MakeArrayView(Schema.Members, Schema.NumMembers))
+	
+	uint16 RangeTypeIdx = 0;
+	for (FMemberType Member : Schema.GetMemberTypes())
 	{
-		Num += (Member.IsStruct() && !Member.AsStruct().IsDynamic);
+		if (Member.IsRange())
+		{
+			FMemberType InnermostType = GrabInnerRangeTypes(RangeTypes, /* in-out */ RangeTypeIdx).Last();
+			Num += IsEnum(InnermostType);
+		}
+		else
+		{
+			Num += IsEnum(Member);
+		}
 	}
+	check(RangeTypeIdx == Schema.NumRangeTypes);
 	return Num;
 }
 
@@ -142,37 +147,46 @@ struct FLoadBatch
 
 using SubsetByteArray = TArray<uint8, TInlineAllocator<1024>>;
 
-static void CopyMemberBinding(FLeafMemberBinding From, FMemberBinder& To)
+static void CopyMemberBinding(FLeafMemberBinding Binding, /* in-out */ const FSchemaId*& InnerSchemaIt, FMemberBinder& Out)
 {
-	// Skip enum schema
-	To.AddMember(From.Leaf.Pack(), static_cast<uint32>(From.Offset));
+	InnerSchemaIt += (Binding.Leaf.Type == ELeafBindType::Enum); // Skip enum schema
+	Out.AddMember(Binding.Leaf.Pack(), static_cast<uint32>(Binding.Offset));
 }
 
-static void CopyMemberBinding(FStructMemberBinding From, FMemberBinder& To)
+static void CopyMemberBinding(FStructMemberBinding Binding, /* in-out */ const FSchemaId*& InnerSchemaIt,FMemberBinder& Out)
 {
-	To.AddMember(FMemberBindType(From.Type), static_cast<uint32>(From.Offset));
-	To.AddInnerSchema(From.Id);
+	Out.AddMember(FMemberBindType(Binding.Type), static_cast<uint32>(Binding.Offset));
+	Out.AddInnerSchema(*InnerSchemaIt);
+	++InnerSchemaIt;
 }
 
-static void CopyMemberBinding(FRangeMemberBinding From, FMemberBinder& To)
+static void CopyMemberBinding(FRangeMemberBinding Binding, /* in-out */ const FSchemaId*& InnerSchemaIt, FMemberBinder& Out)
 {
-	FMemberBindType InnermostType = From.InnerTypes[From.NumRanges - 1];
-	To.AddRange(MakeArrayView(From.RangeBindings, From.NumRanges), InnermostType, static_cast<uint32>(From.Offset));
-	To.AddOptionalInnerSchema(From.InnermostSchema);
+	FMemberBindType InnermostType = Binding.InnerTypes[Binding.NumRanges - 1];
+	Out.AddRange(MakeArrayView(Binding.RangeBindings, Binding.NumRanges), InnermostType, static_cast<uint32>(Binding.Offset));
+	if (InnermostType.IsStruct())
+	{
+		Out.AddInnerSchema(*InnerSchemaIt);
+		++InnerSchemaIt;
+	}
+	else
+	{
+		InnerSchemaIt += (InnermostType.AsLeaf().Bind.Type == ELeafBindType::Enum); // Skip enum schema
+	}
 }
 
-static void CopyMemberBinding(/* in-out */ FMemberVisitor& From, FMemberBinder& To)
+static void CopyMemberBinding(/* in-out */ FMemberVisitor& BindIt, /* in-out */ const FSchemaId*& InnerSchemaIt, FMemberBinder& Out)
 {
-	switch (From.PeekKind())
+	switch (BindIt.PeekKind())
 	{
-		case EMemberKind::Leaf:		CopyMemberBinding(From.GrabLeaf(), To);		break;
-		case EMemberKind::Range:	CopyMemberBinding(From.GrabRange(), To);	break;
-		case EMemberKind::Struct:	CopyMemberBinding(From.GrabStruct(), To);	break;
-		default:					check(false);								break;
+		case EMemberKind::Leaf:		CopyMemberBinding(BindIt.GrabLeaf(), InnerSchemaIt, Out);	break;
+		case EMemberKind::Range:	CopyMemberBinding(BindIt.GrabRange(), InnerSchemaIt, Out);	break;
+		case EMemberKind::Struct:	CopyMemberBinding(BindIt.GrabStruct(), InnerSchemaIt, Out);	break;
+		default:					check(false);											break;
 	}
 }
 
-static void CreateSubsetLoadSchema(const FStructSchema& From, const FStructSchemaBinding& To, TConstArrayView<FMemberId> ToNames,  uint16 NumEnums, SubsetByteArray& Out)
+static void CreateSubsetBindingWithoutEnumIds(const FStructSchema& From, const FStructSchemaBinding& To, TConstArrayView<FMemberId> ToNames,  uint16 NumEnums, SubsetByteArray& Out)
 {
 	check(To.NumMembers == ToNames.Num());
 	check(To.NumMembers >= From.NumMembers);
@@ -184,6 +198,7 @@ static void CreateSubsetLoadSchema(const FStructSchema& From, const FStructSchem
 	
 	FMemberVisitor ToIt(To);
 	FMemberBinder Footer(*Schema);
+	const FSchemaId* InnerSchemaIt = From.GetInnerSchemas();
 	for (FMemberId FromName : From.GetMemberNames())
 	{
 		while (FromName != ToNames[ToIt.GetIndex()])
@@ -191,8 +206,17 @@ static void CreateSubsetLoadSchema(const FStructSchema& From, const FStructSchem
 			ToIt.SkipMember();
 		}
 
-		CopyMemberBinding(/* in-out */ ToIt, /* out */ Footer);	
+		CopyMemberBinding(/* in-out */ ToIt, /* in-out */ InnerSchemaIt, /* out */ Footer);	
 	}
+	check(InnerSchemaIt == From.GetInnerSchemas() + From.NumInnerSchemas);
+}
+
+static void CloneBindingWithReplacedStructIds(const FSchemaId* FromIds, const FStructSchemaBinding& To, SubsetByteArray& Out)
+{
+	int32 OutPos = Out.Num();
+	Out.AddUninitialized(To.CalculateSize());
+	FStructSchemaBinding* Schema = new (&Out[OutPos]) FStructSchemaBinding {To};
+	FMemory::Memcpy(const_cast<FSchemaId*>(Schema->GetInnerSchemas()), FromIds, To.NumInnerSchemas * sizeof(FSchemaId));
 }
 
 [[nodiscard]] static FLoadStructPlan MakeSchemaLoadPlan(const FStructSchema& From, const FStructSchemaBinding& To, TConstArrayView<FMemberId> ToMemberIds, TConstArrayView<FStructSchemaId> ToStructIds, SubsetByteArray& OutSubsetSchemas)
@@ -200,13 +224,19 @@ static void CreateSubsetLoadSchema(const FStructSchema& From, const FStructSchem
 	uint16 NumEnums = CountEnums(From);
 	if (From.NumMembers < To.NumMembers || NumEnums || HasDifferentSupers(From, To, ToStructIds))
 	{
-		CreateSubsetLoadSchema(From, To, ToMemberIds, NumEnums, /* out */ OutSubsetSchemas);
+		CreateSubsetBindingWithoutEnumIds(From, To, ToMemberIds, NumEnums, OutSubsetSchemas);
 	}
-	else // Reuse To bindings
+	else
 	{
 		check(From.NumMembers == To.NumMembers);
 		check(From.NumInnerSchemas == To.NumInnerSchemas);
 		check(From.NumRangeTypes == To.NumInnerRanges);
+
+		if (From.NumInnerSchemas > 0)
+		{
+			CloneBindingWithReplacedStructIds(From.GetInnerSchemas(), To, OutSubsetSchemas);
+		}
+		// else reuse existing bindings
 	}
 
 	// Pointer to created subset load schema will be remapped later
@@ -334,13 +364,6 @@ static FMemberBindType ToBindType(FMemberType Member)
 class FRangeLoader
 {
 public:
-	static void LoadNestedRange(uint8* Member, FMemoryView Values, const FLoadBatch& Batch, const FLoadRangePlan& Range)
-	{
-		FByteReader ByteIt(Values);
-		FBitCacheReader BitIt;
-		LoadRange(Member, ByteIt, BitIt, Batch, Range);
-	}
-
 	static void LoadRangeView(uint8* Member, FRangeView Src, ERangeSizeType MaxSize, TConstArrayView<FRangeBinding> Bindings, const FLoadBatch& Batch)
 	{
 		TArray<FMemberBindType, TFixedAllocator<16>> InnerTypes;
@@ -358,6 +381,7 @@ public:
 		FByteReader ByteIt(Src.Values);
 		FBitCacheReader BitIt;
 		LoadRange(Member, ByteIt, BitIt, Batch, Plan);
+		ByteIt.CheckEmpty();
 	}
 
 	static void LoadRange(uint8* Member, FByteReader& ByteIt, FBitCacheReader& BitIt, const FLoadBatch& Batch, const FLoadRangePlan& Range)
@@ -384,21 +408,37 @@ public:
 	template<class SchemaType>
 	static void LoadRangeValues(uint8* Member, uint64 Num, const IRangeBinding& Binding, FByteReader& ByteIt, const FLoadBatch& Batch, SchemaType&& Schema)
 	{
+		FByteReader ValueIt(GrabRangeValues(ByteIt, Num, Schema));
+		FBitCacheReader BitIt; // Only used by ranges of ERangeSizeType::Uni ranges
 		FLoadRangeContext Ctx{.Request = {Member, Num}};
 		
 		while (Ctx.Request.Index < Num)
 		{
 			(Binding.MakeItems)(Ctx);
-			CopyRangeValues(Ctx.Items, ByteIt, Batch, Schema);
+			CopyRangeValues(Ctx.Items, ValueIt, BitIt, Batch, Schema);
 			Ctx.Request.Index += Ctx.Items.Num;
 		}
-		
+		ValueIt.CheckEmpty();
+
 		if (Ctx.Items.bNeedFinalize)
 		{
 			(Binding.MakeItems)(Ctx);
 		}
 	}
-	static void CopyRangeValues(const FConstructedItems& Items, FByteReader& ByteIt, const FLoadBatch& Batch, FLeafBindType Leaf)
+
+	static FMemoryView GrabRangeValues(FByteReader& ByteIt, uint64 Num, FLeafBindType Schema)
+	{
+		check(Num > 0);
+		return ByteIt.GrabSlice(GetLeafRangeSize(Num, ToLeafType(Schema)));
+	}
+
+	template<class SchemaType>
+	static FMemoryView GrabRangeValues(FByteReader& ByteIt, uint64, SchemaType&&)
+	{
+		return ByteIt.GrabSkippableSlice();
+	}
+		
+	static void CopyRangeValues(const FConstructedItems& Items, FByteReader& ByteIt, FBitCacheReader&, const FLoadBatch& Batch, FLeafBindType Leaf)
 	{
 		switch (Leaf.Bind.Type)
 		{
@@ -422,7 +462,7 @@ public:
 		}
 	}
 	
-	static void CopyRangeValues(const FConstructedItems& Items, FByteReader& ByteIt, const FLoadBatch& Batch, FStructSchemaId Id)
+	static void CopyRangeValues(const FConstructedItems& Items, FByteReader& ByteIt, FBitCacheReader&, const FLoadBatch& Batch, FStructSchemaId Id)
 	{
 		uint64 ItemSize = Items.Size;
 		for (uint8* It = Items.Data, *End = It + Items.NumBytes(); It != End; It += ItemSize)
@@ -431,12 +471,12 @@ public:
 		}
 	}
 	
-	static void CopyRangeValues(const FConstructedItems& Items, FByteReader& ByteIt, const FLoadBatch& Batch, const FLoadRangePlan& Plan)
+	static void CopyRangeValues(const FConstructedItems& Items, FByteReader& ByteIt, FBitCacheReader& BitIt, const FLoadBatch& Batch, const FLoadRangePlan& Plan)
 	{
 		uint64 ItemSize = Items.Size;
 		for (uint8* It = Items.Data, *End = It + Items.NumBytes(); It != End; It += ItemSize)
 		{
-			LoadNestedRange(It, ByteIt.GrabSkippableSlice(), Batch, Plan);	
+			LoadRange(It, ByteIt, BitIt, Batch, Plan);	
 		}
 	}
 };
@@ -455,10 +495,7 @@ public:
 	, RangeBindings(Schema.GetRangeBindings())
 	, Batch(InBatch)
 	, ByteIt(Values)
-	{
-		checkf(CountStaticStructs(Schema) == Schema.NumInnerSchemas, TEXT("Expects Schema stripped from load-irrelevant enum schema ids, see CreateSubsetLoadSchema(). "
-			"# schemas/structs/enums: %d/%d/%d"), Schema.NumInnerSchemas, CountStaticStructs(Schema), CountEnums(Schema));
-	}
+	{}
 
 	void Load(uint8* Struct)
 	{
diff --git a/Engine/Source/Programs/PlainPropsTests/Private/PlainPropsWrite.cpp b/Engine/Source/Programs/PlainPropsTests/Private/PlainPropsWrite.cpp
--- a/Engine/Source/Programs/PlainPropsTests/Private/PlainPropsWrite.cpp
+++ b/Engine/Source/Programs/PlainPropsTests/Private/PlainPropsWrite.cpp
@@ -808,6 +808,16 @@ void FDebugIds::AppendDebugString(FString& Out, FTypeId Type) const
 	AppendDebugString(Out, Type.Name);
 }
 
+void FDebugIds::AppendDebugString(FString& Out, FEnumSchemaId Name) const
+{
+	AppendDebugString(Out, Resolve(Name));
+}
+
+void FDebugIds::AppendDebugString(FString& Out, FStructSchemaId Name) const
+{
+	AppendDebugString(Out, Resolve(Name));
+}
+
 FString FDebugIds::Print(FNameId Name) const
 {
 	FString Out;
@@ -834,4 +844,18 @@ FString FDebugIds::Print(FTypeId Type) const
 	return Out;
 }
 
+FString FDebugIds::Print(FEnumSchemaId Name) const
+{
+	FString Out;
+	AppendDebugString(Out, Name);
+	return Out;
+}
+
+FString FDebugIds::Print(FStructSchemaId Name) const
+{
+	FString Out;
+	AppendDebugString(Out, Name);
+	return Out;
+}
+
 } // namespace PlainProps
diff --git a/Engine/Source/Programs/PlainPropsTests/Public/PlainPropsBind.h b/Engine/Source/Programs/PlainPropsTests/Public/PlainPropsBind.h
--- a/Engine/Source/Programs/PlainPropsTests/Public/PlainPropsBind.h
+++ b/Engine/Source/Programs/PlainPropsTests/Public/PlainPropsBind.h
@@ -556,23 +556,17 @@ constexpr uint32 CountRangeBindings()
 	}
 }
 
-template<typename InnerBinding, typename InnerType>
-struct TGetInnermostImpl
+template<typename RangeBinding, uint32 NestLevel>
+struct TInnerType
 {
-	using Type = typename TGetInnermostImpl<RangeBind<InnerType>, InnerType>::Type;
-};
-
-template<typename InnerType>
-struct TGetInnermostImpl<void, InnerType>
-{
-	using Type = InnerType;
+	using InnerType = typename RangeBinding::ItemType;
+	using Type = TInnerType<RangeBind<InnerType>, NestLevel - 1>::Type;
 };
 
 template<typename RangeBinding>
-struct TGetInnermost
+struct TInnerType<RangeBinding, 1>
 {
-	using InnerType = typename RangeBinding::ItemType;
-	using Type = TGetInnermostImpl<RangeBind<InnerType>, InnerType>::Type;
+	using Type = typename RangeBinding::ItemType;
 };
 
 template<typename RangeBinding, uint32 N>
@@ -614,7 +608,6 @@ TConstArrayView<FRangeBinding> GetRangeBindings()
 	}
 }
 
-
 template<class Var, class Runtime>
 FMemberBinding BindMember()
 {
@@ -635,7 +628,7 @@ FMemberBinding BindMember()
 
 		if constexpr (std::is_void_v<CustomBinding> && NumRangeBindings)
 		{
-			using InnermostType = typename TGetInnermost<RangeBinding>::Type;
+			using InnermostType = typename TInnerType<RangeBinding, NumRangeBindings>::Type;
 			Out.RangeBindings = GetRangeBindings<RangeBinding, NumRangeBindings>();
 			Out.InnermostType = BindType<InnermostType, Ids>(Out.InnermostSchema);
 		}
@@ -764,14 +757,6 @@ struct FMemberBinder
 		*InnerSchemaIt++ = InnermostSchema;
 	}
 
-	void AddOptionalInnerSchema(FOptionalSchemaId InnermostSchema)
-	{
-		if (InnermostSchema)
-		{
-			AddInnerSchema(InnermostSchema.Get());
-		}
-	}
-
 	FStructSchemaBinding& Schema;
 	FMemberBindType* MemberIt;
 	FMemberBindType* RangeTypeIt;
diff --git a/Engine/Source/Programs/PlainPropsTests/Public/PlainPropsBuildSchema.h b/Engine/Source/Programs/PlainPropsTests/Public/PlainPropsBuildSchema.h
--- a/Engine/Source/Programs/PlainPropsTests/Public/PlainPropsBuildSchema.h
+++ b/Engine/Source/Programs/PlainPropsTests/Public/PlainPropsBuildSchema.h
@@ -49,8 +49,9 @@ public:
 	FSchemasBuilder(FStructDeclarations InStructs, FEnumDeclarations InEnums, const FDebugIds& InDebug);
 	~FSchemasBuilder();
 
-	void										NoteMembers(FStructSchemaId Id, const FBuiltStruct& Struct);
-	void										NoteValue(FEnumSchemaId Id, uint64 Value);
+	FEnumSchemaBuilder&							NoteEnum(FEnumSchemaId Id);
+	FStructSchemaBuilder&						NoteStruct(FStructSchemaId Id);
+	void										NoteStructAndMembers(FStructSchemaId Id, const FBuiltStruct& Struct);
 	FBuiltSchemas								Build();
 
 private:
diff --git a/Engine/Source/Programs/PlainPropsTests/Public/PlainPropsIndex.h b/Engine/Source/Programs/PlainPropsTests/Public/PlainPropsIndex.h
--- a/Engine/Source/Programs/PlainPropsTests/Public/PlainPropsIndex.h
+++ b/Engine/Source/Programs/PlainPropsTests/Public/PlainPropsIndex.h
@@ -58,6 +58,8 @@ public:
 	
 	virtual FNestedScope				Resolve(FNestedScopeId Id) const override final		{ return NestedScopes.Resolve(Id); }
 	virtual FParametricTypeView			Resolve(FParametricTypeId Id) const override final	{ return ParametricTypes.Resolve(Id); }
+	virtual FTypeId						Resolve(FEnumSchemaId Id) const override final		{ return Enums.Get(FSetElementId::FromInteger(Id.Idx)); }
+	virtual FTypeId						Resolve(FStructSchemaId Id) const override final	{ return Structs.Get(FSetElementId::FromInteger(Id.Idx)); }
 	
 	const FNestedScopeIndexer&			GetNestedScopes() const								{ return NestedScopes; }
 	const FParametricTypeIndexer&		GetParametricTypes() const							{ return ParametricTypes; }
diff --git a/Engine/Source/Programs/PlainPropsTests/Public/PlainPropsTypes.h b/Engine/Source/Programs/PlainPropsTests/Public/PlainPropsTypes.h
--- a/Engine/Source/Programs/PlainPropsTests/Public/PlainPropsTypes.h
+++ b/Engine/Source/Programs/PlainPropsTests/Public/PlainPropsTypes.h
@@ -343,6 +343,7 @@ template<> inline constexpr FUnpackedLeafType ReflectLeaf<uint32>	= { ELeafType:
 template<> inline constexpr FUnpackedLeafType ReflectLeaf<uint64>	= { ELeafType::IntU,	ELeafWidth::B64 };
 template<> inline constexpr FUnpackedLeafType ReflectLeaf<float>	= { ELeafType::Float,	ELeafWidth::B32 };
 template<> inline constexpr FUnpackedLeafType ReflectLeaf<double>	= { ELeafType::Float,	ELeafWidth::B64 };
+template<> inline constexpr FUnpackedLeafType ReflectLeaf<char>		= { ELeafType::Unicode,	ELeafWidth::B8 };
 template<> inline constexpr FUnpackedLeafType ReflectLeaf<char8_t>	= { ELeafType::Unicode,	ELeafWidth::B8 };
 template<> inline constexpr FUnpackedLeafType ReflectLeaf<char16_t>	= { ELeafType::Unicode,	ELeafWidth::B16 };
 template<> inline constexpr FUnpackedLeafType ReflectLeaf<char32_t>	= { ELeafType::Unicode,	ELeafWidth::B32 };
@@ -364,16 +365,22 @@ public:
 
 	virtual FParametricTypeView			Resolve(FParametricTypeId Id) const = 0;
 	virtual FNestedScope				Resolve(FNestedScopeId Id) const = 0;
-	virtual void						AppendDebugString(FString& Out, FNameId Name) const = 0;
+	virtual FTypeId						Resolve(FEnumSchemaId Id) const = 0;
+	virtual FTypeId						Resolve(FStructSchemaId Id) const = 0;
 
+	virtual void						AppendDebugString(FString& Out, FNameId Name) const = 0;
 	virtual void						AppendDebugString(FString& Out, FScopeId Scope) const;
 	virtual void						AppendDebugString(FString& Out, FTypenameId Typename) const;
 	virtual void						AppendDebugString(FString& Out, FTypeId Type) const;
-
-	PP_API FString					Print(FNameId Name) const;
-	PP_API FString					Print(FMemberId Name) const;
-	PP_API FString					Print(FOptionalMemberId Name) const;
-	PP_API FString					Print(FTypeId Type) const;
+	virtual void						AppendDebugString(FString& Out, FEnumSchemaId Name) const;
+	virtual void						AppendDebugString(FString& Out, FStructSchemaId Name) const;
+
+	PP_API FString						Print(FNameId Name) const;
+	PP_API FString						Print(FMemberId Name) const;
+	PP_API FString						Print(FOptionalMemberId Name) const;
+	PP_API FString						Print(FTypeId Type) const;
+	PP_API FString						Print(FEnumSchemaId Name) const;
+	PP_API FString						Print(FStructSchemaId Name) const;
 };
 
 } // namespace PlainProps
\ No newline at end of file
diff --git a/Engine/Source/Programs/PlainPropsTests/Tests/PlainPropsTests.cpp b/Engine/Source/Programs/PlainPropsTests/Tests/PlainPropsTests.cpp
--- a/Engine/Source/Programs/PlainPropsTests/Tests/PlainPropsTests.cpp
+++ b/Engine/Source/Programs/PlainPropsTests/Tests/PlainPropsTests.cpp
@@ -182,7 +182,7 @@ TArray64<uint8> FTestBatchBuilder::Write()
 	FSchemasBuilder SchemaBuilders(DeclaredStructs, DeclaredEnums, *this);
 	for (const TPair<FStructSchemaId, TUniquePtr<FBuiltStruct>>& Object : Objects)
 	{
-		SchemaBuilders.NoteMembers(Object.Key, *Object.Value);
+		SchemaBuilders.NoteStructAndMembers(Object.Key, *Object.Value);
 	}
 	FBuiltSchemas Schemas = SchemaBuilders.Build(); 
 
diff --git a/Engine/Source/Programs/PlainPropsTests/Tests/PlainPropsUeCoreTests.cpp b/Engine/Source/Programs/PlainPropsTests/Tests/PlainPropsUeCoreTests.cpp
--- a/Engine/Source/Programs/PlainPropsTests/Tests/PlainPropsUeCoreTests.cpp
+++ b/Engine/Source/Programs/PlainPropsTests/Tests/PlainPropsUeCoreTests.cpp
@@ -47,16 +47,6 @@ struct TCustomBindings
 	using Type = void;
 };
 
-template<class T>
-struct TCustomDeltaBindings : TCustomBindings<T>
-{};
-
-template<typename T>
-struct TCustomDeltaBindings<TSet<T>>
-{
-	using Type = PlainProps::UE::TSetDeltaBinding<FIds, T>;
-};
-
 ////////////////////////////////////////////////////////////////////////////////////////////////
 
 struct FDefaultRuntime
@@ -68,30 +58,18 @@ struct FDefaultRuntime
 	static FStructBindings&			GetBindings()					{ return GBindings; }
 };
 
-struct FDeltaRuntime : FDefaultRuntime
-{
-	template<class T> using CustomBindings = TCustomDeltaBindings<T>;
-
-	static FStructBindings&			GetBindings()					{ return GDeltaBindings; }
-};
-
 ////////////////////////////////////////////////////////////////////////////////////////////////
 
-template<class T, class Runtime>
+template<class T, class Runtime = FDefaultRuntime>
 struct TScopedStructBinding
 {
 	using Ids = typename Runtime::Ids;
 	using Ctti = CttiOf<T>;
-	//constexpr bool bDeltaBindMembers = NumCustomMembers<Ctti, FDeltaRuntime::CustomBindings>() > NumCustomMembers<Ctti, FDefaultRuntime::CustomBindings>();
 
 	TScopedStructBinding(EMemberPresence Occupancy = EMemberPresence::AllowSparse)
 	: Id(DeclareNativeStruct<Ctti, Ids>(Runtime::GetDeclarations(), Occupancy))
 	{
 		BindNativeStruct<Ctti, Runtime>(Runtime::GetBindings(), Id);
-		//if constexpr (bDeltaBindMembers)
-		//{
-		//	BindNativeStruct<Ctti, FDeltaRuntime>(Runtime::GetBindings());	
-		//}
 	}
 
 	~TScopedStructBinding()
@@ -103,7 +81,7 @@ struct TScopedStructBinding
 	FStructSchemaId Id;
 };
 
-template<typename Enum, EEnumMode Mode, class Runtime>
+template<typename Enum, EEnumMode Mode, class Runtime = FDefaultRuntime>
 struct TScopedEnumBinding
 {
 	using Ids = typename Runtime::Ids;
@@ -114,26 +92,6 @@ struct TScopedEnumBinding
 	~TScopedEnumBinding() { Runtime::GetDeclarations().DropEnum(Id); }
 };
 
-//template<typename T, class Runtime>
-//struct TScopedBinding : std::conditional_t<std::is_enum_v<T>, TScopedEnumBinding<T, Runtime>, TScopedStructBinding<T, Runtime>>
-//{
-//	
-//};
-
-//template<class T>
-//struct TScopedStructDeclaration
-//{
-//	FStructSchemaId Id;
-//
-//	TScopedStructDeclaration(EMemberPresence Occupancy)
-//	: Id(DeclareNativeStruct<CttiOf<T>>(FDefaultRuntime::GetDeclarations()), Occupancy)
-//	{}
-//
-//	~TScopedStructDeclaration()
-//	{
-//		FDeltaRuntime::DropStruct(Id);
-//	}
-//};
 ////////////////////////////////////////////////////////////////////////////////////////////////
 
 inline constexpr uint32 Magics[] = { 0xFEEDF00D, 0xABCD1234, 0xDADADAAA, 0x99887766, 0xF0F1F2F3 };
@@ -180,7 +138,7 @@ TArray64<uint8> FBatchSaver::Write() const
 	FSchemasBuilder SchemaBuilders(GTypes.GetStructs(), GTypes.GetEnums(), /* debug */ GNames);
 	for (const IdBuiltStructPair& Object : SavedObjects)
 	{
-		SchemaBuilders.NoteMembers(Object.Key, *Object.Value);
+		SchemaBuilders.NoteStructAndMembers(Object.Key, *Object.Value);
 	}
 	FBuiltSchemas Schemas = SchemaBuilders.Build(); 
 
@@ -306,43 +264,14 @@ static void SaveAndLoad(const FStructBindings& Bindings, void (*Save)(FBatchSave
 	Load(Batch);
 }
 
-template<class Runtime>
+template<class Runtime = FDefaultRuntime>
 static void TestSaveAndLoad(void (*Save)(FBatchSaver&), void (*Load)(FBatchLoader&))
 {
 	SaveAndLoad(Runtime::GetBindings(), Save, Load);
 }
 
+//////////////////////////////////////////////////////////////////////////
 
-
-//// Tests that everything was visited
-//struct FTestMemberVisitor : public FMemberVisitor
-//{
-//	using FMemberVisitor::FMemberVisitor;
-//	
-//	~FMemberVisitor()
-//	{
-//		CHECK(MemberIdx == NumMembers); // Must read all members
-//		CHECK(RangeTypeIdx == NumRangeTypes); // Must read all ranges
-//#if DO_CHECK
-//		CHECK(InnerSchemaIdx == NumInnerSchemas); // Must read all schema ids
-//#endif
-//	}
-//};
-//
-//template<typename OutType, typename InType>
-//TArray<OutType> MakeArray(const InType& Items)
-//{
-//	TArray<OutType> Out;
-//	Out.Reserve(IntCastChecked<int32>(Items.Num()));
-//	for (const auto& Item : Items)
-//	{
-//		Out.Emplace(Item);
-//	}
-//	return Out;
-//}
-//
-//
-//
 //struct FNestedStructs
 //{
 //	FPt Point;
@@ -379,12 +308,9 @@ static void TestSaveAndLoad(void (*Save)(FBatchSaver&), void (*Load)(FBatchLoade
 //	TArray<FSub> Subs;
 //};
 
-
 struct FInt { int32 X; };
 PP_REFLECT_STRUCT(PlainProps::UE::Test, FInt, void, X);
 
-//////////////////////////////////////////////////////////////////////////
-
 enum class EFlat1 : uint8 { A = 1, B = 3 };
 enum class EFlat2 : uint8 { A, B };
 enum class EFlag1 : uint8 { A = 2, B = 8, AB = 10 };
@@ -405,22 +331,6 @@ struct FEnums
 };
 PP_REFLECT_STRUCT(PlainProps::UE::Test, FEnums, void, Flat1, Flat2, Flag1, Flag2);
 
-//////////////////////////////////////////////////////////////////////////
-
-//struct FLeaves
-//{
-//	uint16 U16 = 0;
-//	int64 I64 = 0;
-//	bool Bool = false;
-//	char AsciiChar = 0;
-////	uint8 BitfieldBool : 1;
-//	uint32 StaticU32s[3] = {};
-//	double Double = 0.0;
-//};
-//PP_REFLECT_STRUCT(PlainProps::UE::Test, FLeaves, void, U16, I64, Bool, AsciiChar, StaticU32s, Double);
-
-//////////////////////////////////////////////////////////////////////////
-
 struct FLeafArrays
 {
 	TArray<bool> Bits;
@@ -447,8 +357,8 @@ TEST_CASE_NAMED(FPlainPropsUeCoreTest, "System::Core::Serialization::PlainProps:
 {
 	SECTION("Basic")
 	{
-		TScopedStructBinding<FInt, FDefaultRuntime> Int;
-		TestSaveAndLoad<FDefaultRuntime>(
+		TScopedStructBinding<FInt> Int;
+		TestSaveAndLoad(
 			[](FBatchSaver& Batch)
 			{
 				Batch.Save(FInt{1234});
@@ -461,12 +371,12 @@ TEST_CASE_NAMED(FPlainPropsUeCoreTest, "System::Core::Serialization::PlainProps:
 	
 	SECTION("Enum")
 	{
-		TScopedEnumBinding<EFlat1, EEnumMode::Flat, FDefaultRuntime> Flat1;
-		TScopedEnumBinding<EFlat2, EEnumMode::Flat, FDefaultRuntime> Flat2;
-		TScopedEnumBinding<EFlag1, EEnumMode::Flag, FDefaultRuntime> Flag1;
-		TScopedEnumBinding<EFlag2, EEnumMode::Flag, FDefaultRuntime> Flag2;
-		TScopedStructBinding<FEnums, FDefaultRuntime> Int;
-		TestSaveAndLoad<FDefaultRuntime>(
+		TScopedEnumBinding<EFlat1, EEnumMode::Flat> Flat1;
+		TScopedEnumBinding<EFlat2, EEnumMode::Flat> Flat2;
+		TScopedEnumBinding<EFlag1, EEnumMode::Flag> Flag1;
+		TScopedEnumBinding<EFlag2, EEnumMode::Flag> Flag2;
+		TScopedStructBinding<FEnums> Int;
+		TestSaveAndLoad(
 			[](FBatchSaver& Batch)
 			{
 				Batch.Save(FEnums{EFlat1::A, EFlat2::A, EFlag1::A, EFlag2::A});
@@ -487,24 +397,41 @@ TEST_CASE_NAMED(FPlainPropsUeCoreTest, "System::Core::Serialization::PlainProps:
 
 	SECTION("TArrayBasic")
 	{
-		TScopedStructBinding<FLeafArrays, FDefaultRuntime> LeafArrays;
-		TestSaveAndLoad<FDefaultRuntime>(
+		TScopedStructBinding<FLeafArrays> LeafArrays;
+		TestSaveAndLoad(
 			[](FBatchSaver& Batch)
 			{
 				Batch.Save(FLeafArrays{{}, {}});
 				Batch.Save(FLeafArrays{{false}, {1, 2}});
 				Batch.Save(FLeafArrays{{true, false}, {3, 4, 5}});
+				Batch.Save(FLeafArrays{{true,true,true,true,true,true,true,true,false,true}});
 			}, 
 			[](FBatchLoader& Batch)
 			{
 				CHECK(Batch.Load<FLeafArrays>() == FLeafArrays{{}, {}});
 				CHECK(Batch.Load<FLeafArrays>() == FLeafArrays{{false}, {1, 2}});
 				CHECK(Batch.Load<FLeafArrays>() == FLeafArrays{{true, false}, {3, 4, 5}});
+				CHECK(Batch.Load<FLeafArrays>() == FLeafArrays{{true,true,true,true,true,true,true,true,false,true}});
 			});
 	}
 
 	SECTION("TArrayComplex")
-	{}
+	{
+		TScopedEnumBinding<EFlat1, EEnumMode::Flat> Flat1;
+		TScopedStructBinding<FLeafArrays> LeafArrays;
+		TScopedStructBinding<FComplexArrays> ComplexArrays;
+		TestSaveAndLoad(
+			[](FBatchSaver& Batch)
+			{
+				Batch.Save(FComplexArrays{});
+				Batch.Save(FComplexArrays{{'a', 'b'}, {EFlat1::A}, {{}, {{true}, {2}}}, {{EFlat1::B}, {}} });
+			}, 
+			[](FBatchLoader& Batch)
+			{
+				CHECK(Batch.Load<FComplexArrays>() == FComplexArrays{});
+				CHECK(Batch.Load<FComplexArrays>() == FComplexArrays{{'a', 'b'}, {EFlat1::A}, {{}, {{true}, {2}}}, {{EFlat1::B}, {}} });
+			});
+	}
 
 	SECTION("FString")
 	{}
