PlainProps add TUniquePtr bindings and improve declaration/bind error messages

#rnx

diff --git a/Engine/Source/Programs/PlainPropsTests/Private/PlainPropsBind.cpp b/Engine/Source/Programs/PlainPropsTests/Private/PlainPropsBind.cpp
--- a/Engine/Source/Programs/PlainPropsTests/Private/PlainPropsBind.cpp
+++ b/Engine/Source/Programs/PlainPropsTests/Private/PlainPropsBind.cpp
@@ -223,11 +223,13 @@ void FStructBindings::Bind(FStructSchemaId Id, FStructBinding Binding)
 
 FStructBinding FStructBindings::Get(FStructSchemaId Id) const
 {
+	checkf(Id.Idx < (uint32)Bindings.Num() && Bindings[Id.Idx], TEXT("'%s' is unbound"), *Debug.Print(Id));
 	return Bindings[Id.Idx].Get();
 }
 
 void FStructBindings::DropStruct(FStructSchemaId Id)
 {
+	checkf(Id.Idx < (uint32)Bindings.Num() && Bindings[Id.Idx], TEXT("'%s' is unbound"), *Debug.Print(Id));
 	Bindings[Id.Idx].ResetOwned();
 }
 
diff --git a/Engine/Source/Programs/PlainPropsTests/Private/PlainPropsDeclare.cpp b/Engine/Source/Programs/PlainPropsTests/Private/PlainPropsDeclare.cpp
--- a/Engine/Source/Programs/PlainPropsTests/Private/PlainPropsDeclare.cpp
+++ b/Engine/Source/Programs/PlainPropsTests/Private/PlainPropsDeclare.cpp
@@ -77,4 +77,16 @@ void FDeclarations::DeclareEnum(FEnumSchemaId Id, FTypeId Type, EEnumMode Mode,
 	ValidateDeclaration(*Ptr);
 }
 
+#if DO_CHECK
+void FDeclarations::Check(FEnumSchemaId Id) const
+{
+	checkf(Id.Idx < (uint32)DeclaredEnums.Num() && DeclaredEnums[Id.Idx], TEXT("'%s' is undeclared"), *Debug.Print(Id));
+}
+
+void FDeclarations::Check(FStructSchemaId Id) const
+{
+	checkf(Id.Idx < (uint32)DeclaredStructs.Num() && DeclaredStructs[Id.Idx], TEXT("'%s' is undeclared"), *Debug.Print(Id));
+}
+#endif
+
 } // namespace PlainProps
\ No newline at end of file
diff --git a/Engine/Source/Programs/PlainPropsTests/Private/PlainPropsLoad.cpp b/Engine/Source/Programs/PlainPropsTests/Private/PlainPropsLoad.cpp
--- a/Engine/Source/Programs/PlainPropsTests/Private/PlainPropsLoad.cpp
+++ b/Engine/Source/Programs/PlainPropsTests/Private/PlainPropsLoad.cpp
@@ -213,9 +213,10 @@ static void CreateSubsetBindingWithoutEnumIds(const FStructSchema& From, const F
 
 static void CloneBindingWithReplacedStructIds(const FSchemaId* FromIds, const FStructSchemaBinding& To, SubsetByteArray& Out)
 {
-	int32 OutPos = Out.Num();
-	Out.AddUninitialized(To.CalculateSize());
-	FStructSchemaBinding* Schema = new (&Out[OutPos]) FStructSchemaBinding {To};
+	uint32 Size = To.CalculateSize();
+	Out.AddUninitialized(Size);
+	FStructSchemaBinding* Schema = reinterpret_cast<FStructSchemaBinding*>(&Out[Out.Num() - Size]);
+	FMemory::Memcpy(Schema, &To, Size);
 	FMemory::Memcpy(const_cast<FSchemaId*>(Schema->GetInnerSchemas()), FromIds, To.NumInnerSchemas * sizeof(FSchemaId));
 }
 
diff --git a/Engine/Source/Programs/PlainPropsTests/Public/PlainPropsBind.h b/Engine/Source/Programs/PlainPropsTests/Public/PlainPropsBind.h
--- a/Engine/Source/Programs/PlainPropsTests/Public/PlainPropsBind.h
+++ b/Engine/Source/Programs/PlainPropsTests/Public/PlainPropsBind.h
@@ -438,7 +438,7 @@ class FStructBindings
 {
 public:
 	UE_NONCOPYABLE(FStructBindings);
-	FStructBindings() = default;
+	explicit FStructBindings(const FDebugIds& In) : Debug(In) {}
 	~FStructBindings();
 
 	void							BindStruct(FStructSchemaId Id, const ICustomStructBinding& Custom);
@@ -448,6 +448,7 @@ public:
 
 private:
 	TArray<FStructBindingOwner>		Bindings;
+	const FDebugIds&				Debug;
 
 	void							Bind(FStructSchemaId Id, FStructBinding Binding);
 };
diff --git a/Engine/Source/Programs/PlainPropsTests/Public/PlainPropsDeclare.h b/Engine/Source/Programs/PlainPropsTests/Public/PlainPropsDeclare.h
--- a/Engine/Source/Programs/PlainPropsTests/Public/PlainPropsDeclare.h
+++ b/Engine/Source/Programs/PlainPropsTests/Public/PlainPropsDeclare.h
@@ -51,14 +51,17 @@ struct FStructDeclaration
 class FDeclarations
 {
 public:
+	UE_NONCOPYABLE(FDeclarations);
+	explicit FDeclarations(const FDebugIds& In) : Debug(In) {}
+
 	void											DeclareEnum(FEnumSchemaId Id, FTypeId Type, EEnumMode Mode, ELeafWidth Width, TConstArrayView<FEnumerator> Enumerators);
 	void											DeclareStruct(FStructSchemaId Id, FTypeId Type, TConstArrayView<FMemberId> MemberOrder, EMemberPresence Occupancy, FOptionalStructSchemaId Super = {});
 	
-	void											DropEnum(FEnumSchemaId Id)		{ Get(Id); DeclaredEnums[Id.Idx].Reset(); }
-	void											DropStruct(FStructSchemaId Id)	{ Get(Id); DeclaredStructs[Id.Idx].Reset(); }
+	void											DropEnum(FEnumSchemaId Id)		{ Check(Id); DeclaredEnums[Id.Idx].Reset(); }
+	void											DropStruct(FStructSchemaId Id)	{ Check(Id); DeclaredStructs[Id.Idx].Reset(); }
 
-	const FEnumDeclaration&							Get(FEnumSchemaId Id) const		{ check(DeclaredEnums[Id.Idx]);		return *DeclaredEnums[Id.Idx]; }
-	const FStructDeclaration&						Get(FStructSchemaId Id) const	{ check(DeclaredStructs[Id.Idx]);	return *DeclaredStructs[Id.Idx]; }
+	const FEnumDeclaration&							Get(FEnumSchemaId Id) const		{ Check(Id); return *DeclaredEnums[Id.Idx]; }
+	const FStructDeclaration&						Get(FStructSchemaId Id) const	{ Check(Id); return *DeclaredStructs[Id.Idx]; }
 	
 	TConstArrayView<TUniquePtr<FEnumDeclaration>>	GetEnums() const	{ return DeclaredEnums; }
 	TConstArrayView<TUniquePtr<FStructDeclaration>>	GetStructs() const	{ return DeclaredStructs; }
@@ -66,6 +69,14 @@ public:
 protected:
 	TArray<TUniquePtr<FEnumDeclaration>>			DeclaredEnums;
 	TArray<TUniquePtr<FStructDeclaration>>			DeclaredStructs;
+	const FDebugIds&								Debug;
+
+#if DO_CHECK
+	void											Check(FEnumSchemaId Id) const;
+	void											Check(FStructSchemaId Id) const;
+#else
+	void											Check(...) const {}
+#endif
 };
 
 } // namespace PlainProps
\ No newline at end of file
diff --git a/Engine/Source/Programs/PlainPropsTests/Public/PlainPropsUeCoreBindings.h b/Engine/Source/Programs/PlainPropsTests/Public/PlainPropsUeCoreBindings.h
--- a/Engine/Source/Programs/PlainPropsTests/Public/PlainPropsUeCoreBindings.h
+++ b/Engine/Source/Programs/PlainPropsTests/Public/PlainPropsUeCoreBindings.h
@@ -130,10 +130,53 @@ struct TArrayBinding : public IRangeBinding
 
 //////////////////////////////////////////////////////////////////////////
 
+template <typename T>
+struct TUniquePtrBinding : public IRangeBinding
+{
+	using SizeType = bool;
+	using ItemType = T;
+
+	virtual void MakeItems(FLoadRangeContext& Ctx) const override
+	{
+		TUniquePtr<T>& Ptr = Ctx.Request.GetRange<TUniquePtr<T>>();
+		
+		if (Ctx.Request.NumTotal() == 0)
+		{
+			Ptr.Reset();
+			return;
+		}
+		
+		if (!Ptr)
+		{
+			if constexpr (std::is_default_constructible_v<T>)
+			{
+				Ptr.Reset(new T);
+			}
+			else
+			{
+				Ptr.Reset(reinterpret_cast<T*>(FMemory::Malloc(sizeof(T), alignof(T))));
+				Ctx.Items.SetUnconstructed();
+			}
+		}
+		
+		Ctx.Items.Set(Ptr.Get(), 1);
+	}
+
+	virtual void ReadItems(FSaveRangeContext& Ctx) const override
+	{
+		const TUniquePtr<T>& Ptr = Ctx.Request.GetRange<TUniquePtr<T>>();
+		Ctx.Items.SetAll(Ptr.Get(), Ptr ? 1 : 0);
+	}
+
+};
+
+//////////////////////////////////////////////////////////////////////////
+
 template <typename T>
 struct TSetBinding : public IRangeBinding
 {
 	using SizeType = int32;
+	using ItemType = T;
 
 	virtual void MakeItems(FLoadRangeContext& Ctx) const override
 	{
@@ -436,6 +479,12 @@ namespace PlainProps
 		using Type = UE::TArrayBinding<T>;
 	};
 
+	template<typename T>
+	struct TRangeBind<TUniquePtr<T>>
+	{
+		using Type = UE::TUniquePtrBinding<T>;
+	};
+
 	template<typename T>
 	struct TRangeBind<TSet<T>>
 	{
diff --git a/Engine/Source/Programs/PlainPropsTests/Tests/PlainPropsTests.cpp b/Engine/Source/Programs/PlainPropsTests/Tests/PlainPropsTests.cpp
--- a/Engine/Source/Programs/PlainPropsTests/Tests/PlainPropsTests.cpp
+++ b/Engine/Source/Programs/PlainPropsTests/Tests/PlainPropsTests.cpp
@@ -104,13 +104,18 @@ TEST_CASE_NAMED(FPlainPropsIndexTest, "System::Core::Serialization::PlainProps::
 
 inline constexpr uint32 TestMagics[] = {0xFEEDF00D, 0xABCD1234, 0xDADADAAA, 0x99887766, 0xF0F1F2F3 };
 
-class FTestBatchBuilder : public TIdIndexer<FAnsiString>, public FDeclarations
+class FTestBatchBuilder : public TIdIndexer<FAnsiString>
 {
 public:
+	FTestBatchBuilder() : Declarations(*this) {}
+
 	FEnumSchemaId				DeclareEnum(FTypeId Type, EEnumMode Mode, ELeafWidth Width, std::initializer_list<const char*> Names, std::initializer_list<uint64> Constants);
 	FEnumSchemaId				DeclareEnum(const char* Scope, const char* Name, EEnumMode Mode, ELeafWidth Width, std::initializer_list<const char*> Names, std::initializer_list<uint64> Constants);
 	FStructSchemaId				DeclareStruct(FTypeId Type, std::initializer_list<const char*> MemberOrder, EMemberPresence Occupancy, FOptionalStructSchemaId Super = {});
 	FStructSchemaId				DeclareStruct(const char* Scope, const char* Name, std::initializer_list<const char*> MemberOrder, EMemberPresence Occupancy, FOptionalStructSchemaId Super = {});
+	
+	const FEnumDeclaration&		Get(FEnumSchemaId Id) { return Declarations.Get(Id); }
+	const FStructDeclaration&	Get(FStructSchemaId Id) { return Declarations.Get(Id); }
 
 	void						AddObject(FStructSchemaId Schema, FMemberBuilder&& Members);
 	void						AddObjects(FStructSchemaId Schema, std::initializer_list<FMemberBuilder&&> Objects);
@@ -119,6 +124,7 @@ public:
 
 private:
 	TArray<TPair<FStructSchemaId, TUniquePtr<FBuiltStruct>>>	Objects;
+	FDeclarations												Declarations;
 
 	TArray<FMemberId> NameMembers(std::initializer_list<const char*> Members)
 	{
@@ -149,7 +155,7 @@ private:
 FStructSchemaId FTestBatchBuilder::DeclareStruct(FTypeId Type, std::initializer_list<const char*> MemberOrder, EMemberPresence Occupancy, FOptionalStructSchemaId Super)
 {
 	FStructSchemaId Id = IndexStruct(Type);
-	FDeclarations::DeclareStruct(Id, Type, NameMembers(MemberOrder), Occupancy, Super);
+	Declarations.DeclareStruct(Id, Type, NameMembers(MemberOrder), Occupancy, Super);
 	return Id;
 }
 
@@ -161,7 +167,7 @@ FStructSchemaId FTestBatchBuilder::DeclareStruct(const char* Scope, const char*
 FEnumSchemaId FTestBatchBuilder::DeclareEnum(FTypeId Type, EEnumMode Mode, ELeafWidth Width, std::initializer_list<const char*> InNames, std::initializer_list<uint64> Constants)
 {
 	FEnumSchemaId Id = IndexEnum(Type);
-	FDeclarations::DeclareEnum(Id, Type, Mode, Width, MakeEnumerators(InNames, Constants));
+	Declarations.DeclareEnum(Id, Type, Mode, Width, MakeEnumerators(InNames, Constants));
 	return Id;
 }
 
@@ -172,14 +178,13 @@ FEnumSchemaId FTestBatchBuilder::DeclareEnum(const char* Scope, const char* Name
 
 void FTestBatchBuilder::AddObject(FStructSchemaId Schema, FMemberBuilder&& Members)
 {
-	const FStructDeclaration& Declaration = *DeclaredStructs[Schema.Idx];
-	Objects.Emplace(Schema, Members.BuildAndReset(Declaration, *this));
+	Objects.Emplace(Schema, Members.BuildAndReset(Declarations.Get(Schema), *this));
 }
 
 TArray64<uint8> FTestBatchBuilder::Write()
 {
 	// Build partial schemas
-	FSchemasBuilder SchemaBuilders(DeclaredStructs, DeclaredEnums, *this);
+	FSchemasBuilder SchemaBuilders(Declarations.GetStructs(), Declarations.GetEnums(), *this);
 	for (const TPair<FStructSchemaId, TUniquePtr<FBuiltStruct>>& Object : Objects)
 	{
 		SchemaBuilders.NoteStructAndMembers(Object.Key, *Object.Value);
diff --git a/Engine/Source/Programs/PlainPropsTests/Tests/PlainPropsUeCoreTests.cpp b/Engine/Source/Programs/PlainPropsTests/Tests/PlainPropsUeCoreTests.cpp
--- a/Engine/Source/Programs/PlainPropsTests/Tests/PlainPropsUeCoreTests.cpp
+++ b/Engine/Source/Programs/PlainPropsTests/Tests/PlainPropsUeCoreTests.cpp
@@ -12,18 +12,18 @@
 #include  "PlainPropsSave.h"
 #include  "PlainPropsWrite.h"
 #include  "PlainPropsUeCoreBindings.h"
-#include "Tests/TestHarnessAdapter.h"
+#include "Algo/Compare.h"
 #include "Containers/StringView.h"
 #include "Containers/Map.h"
 #include "Templates/UnrealTemplate.h"
+#include "Tests/TestHarnessAdapter.h"
 
 namespace PlainProps::UE::Test
 {
 
 static TIdIndexer<FName>	GNames;
-static FDeclarations		GTypes;
-static FStructBindings		GBindings;
-static FStructBindings		GDeltaBindings;
+static FDeclarations		GTypes(/* debug */ GNames);
+static FStructBindings		GBindings(/* debug */ GNames);
 
 struct FIds
 {
@@ -310,6 +310,7 @@ static void TestSaveAndLoad(void (*Save)(FBatchSaver&), void (*Load)(FBatchLoade
 
 struct FInt { int32 X; };
 PP_REFLECT_STRUCT(PlainProps::UE::Test, FInt, void, X);
+static bool operator==(FInt A, FInt B) { return A.X == B.X; }
 
 enum class EFlat1 : uint8 { A = 1, B = 3 };
 enum class EFlat2 : uint8 { A, B };
@@ -326,19 +327,17 @@ struct FEnums
 	EFlat2 Flat2;
 	EFlag1 Flag1;
 	EFlag2 Flag2;
-
-	bool operator==(FEnums O) const { return Flat1 == O.Flat1 && Flat2 == O.Flat2 && Flag1 == O.Flag1 && Flag2 == O.Flag2; }
 };
 PP_REFLECT_STRUCT(PlainProps::UE::Test, FEnums, void, Flat1, Flat2, Flag1, Flag2);
+static bool operator==(FEnums A, FEnums B) { return A.Flat1 == B.Flat1 && A.Flat2 == B.Flat2 && A.Flag1 == B.Flag1 && A.Flag2 == B.Flag2; }
 
 struct FLeafArrays
 {
 	TArray<bool> Bits;
 	TArray<int>	Bobs;
-
-	bool operator==(const FLeafArrays& O) const { return Bits == O.Bits && Bobs == O.Bobs; }
 };
 PP_REFLECT_STRUCT(PlainProps::UE::Test, FLeafArrays, void, Bits, Bobs);
+static bool operator==(const FLeafArrays& A, const FLeafArrays& B) { return A.Bits == B.Bits && A.Bobs == B.Bobs; }
 
 struct FComplexArrays
 {
@@ -346,10 +345,47 @@ struct FComplexArrays
 	TArray<EFlat1> Enums;
 	TArray<FLeafArrays> Misc;
 	TArray<TArray<EFlat1>> Nested;
-
-	bool operator==(const FComplexArrays& O) const { return Str == O.Str && Enums == O.Enums && Misc == O.Misc && Nested == O.Nested; }
 };
 PP_REFLECT_STRUCT(PlainProps::UE::Test, FComplexArrays, void, Str, Enums, Misc, Nested);
+static bool operator==(const FComplexArrays& A, const FComplexArrays& B) { return A.Str == B.Str && A.Enums == B.Enums && A.Misc == B.Misc && A.Nested == B.Nested; }
+
+struct FUniquePtrs
+{
+	TUniquePtr<bool> Bit;
+	TUniquePtr<FInt> Struct;
+	TUniquePtr<TUniquePtr<int>> IntPtr;
+	TArray<TUniquePtr<double>> Doubles;
+};
+PP_REFLECT_STRUCT(PlainProps::UE::Test, FUniquePtrs, void, Bit, Struct, IntPtr, Doubles);
+
+template<typename T>
+bool SameValue(const TUniquePtr<T>& A, const TUniquePtr<T>& B)
+{
+	return !A == !B && (!A || *A == *B); 
+}
+
+static bool operator==(const FUniquePtrs& A, const FUniquePtrs& B) 
+{ 
+	return	SameValue(A.Bit, B.Bit) &&
+			SameValue(A.Struct, B.Struct) &&
+			!A.IntPtr == !B.IntPtr && (!A.IntPtr || SameValue(*A.IntPtr, *B.IntPtr)) && 
+			Algo::Compare(A.Doubles, B.Doubles, [](auto& X, auto& Y) { return SameValue(X, Y); });
+}
+
+template<typename T>
+TUniquePtr<T> MakeOne(T&& Value)
+{
+	return MakeUnique<T>(MoveTemp(Value));
+}
+
+template<typename T>
+TArray<TUniquePtr<T>> MakeTwo(T&& A, T&& B)
+{
+	TArray<TUniquePtr<T>> Out;
+	Out.Add(MakeOne(MoveTemp(A)));
+	Out.Add(MakeOne(MoveTemp(B)));
+	return Out;
+}
 
 //////////////////////////////////////////////////////////////////////////
 
@@ -433,10 +469,27 @@ TEST_CASE_NAMED(FPlainPropsUeCoreTest, "System::Core::Serialization::PlainProps:
 			});
 	}
 
-	SECTION("FString")
+	SECTION("TUniquePtr")
+	{
+		TScopedStructBinding<FInt> Int;
+		TScopedStructBinding<FUniquePtrs> UniquePtrs;
+		TestSaveAndLoad(
+			[](FBatchSaver& Batch)
+			{
+				Batch.Save(FUniquePtrs{});
+				Batch.Save(FUniquePtrs{MakeOne(true), MakeOne(FInt{3}), MakeOne(MakeOne(2)), MakeTwo(1.0, 2.0)});
+			}, 
+			[](FBatchLoader& Batch)
+			{
+				CHECK(Batch.Load<FUniquePtrs>() == FUniquePtrs{});
+				CHECK(Batch.Load<FUniquePtrs>() == FUniquePtrs{MakeOne(true), MakeOne(FInt{3}), MakeOne(MakeOne(2)), MakeTwo(1.0, 2.0)});
+			});
+	}
+
+	SECTION("FakeReference")
 	{}
 
-	SECTION("TUniquePtr")
+	SECTION("FString")
 	{}
 		
 	SECTION("TSet")
@@ -445,9 +498,6 @@ TEST_CASE_NAMED(FPlainPropsUeCoreTest, "System::Core::Serialization::PlainProps:
 	SECTION("FName")
 	{}
 
-	SECTION("FakeReference")
-	{}
-
 	SECTION("NestedContainer")
 	{}
 	
