PlainProps custom bindings

#rnx

diff --git a/Engine/Source/Programs/PlainPropsTests/Private/PlainPropsBind.cpp b/Engine/Source/Programs/PlainPropsTests/Private/PlainPropsBind.cpp
--- a/Engine/Source/Programs/PlainPropsTests/Private/PlainPropsBind.cpp
+++ b/Engine/Source/Programs/PlainPropsTests/Private/PlainPropsBind.cpp
@@ -19,9 +19,9 @@ static_assert((uint8)ELeafType::Unicode		== (uint8)ELeafBindType::Unicode);
 
 ////////////////////////////////////////////////////////////////////////////////////////////////
 
-uint32 FStructSchemaBinding::CalculateSize() const
+uint32 FSchemaBinding::CalculateSize() const
 {
-	uint32 Out = sizeof(FStructSchemaBinding) + (NumMembers + NumInnerRanges) * sizeof(FMemberBindType);
+	uint32 Out = sizeof(FSchemaBinding) + (NumMembers + NumInnerRanges) * sizeof(FMemberBindType);
 	Out = Align(Out + NumMembers * sizeof(uint32), sizeof(uint32));
 	Out = Align(Out + NumInnerSchemas * sizeof(FSchemaId), sizeof(FSchemaId));
 	Out = Align(Out + NumInnerRanges * sizeof(FRangeBinding), sizeof(FRangeBinding));
@@ -30,7 +30,7 @@ uint32 FStructSchemaBinding::CalculateSize() const
 
 ////////////////////////////////////////////////////////////////////////////////////////////////
 
-FMemberVisitor::FMemberVisitor(const FStructSchemaBinding& InSchema)
+FMemberVisitor::FMemberVisitor(const FSchemaBinding& InSchema)
 : Schema(InSchema)
 , NumMembers(InSchema.NumMembers)
 {}
@@ -122,44 +122,43 @@ FRangeBinding::FRangeBinding(const IRangeBinding& Binding, ERangeSizeType SizeTy
 
 ////////////////////////////////////////////////////////////////////////////////////////////////
 
-FStructBindingOwner::~FStructBindingOwner()
+void FCustomBindings::BindStruct(FStructSchemaId Id, ICustomBinding& Binding)
 {
-	if (Handle & FStructBinding::SchemaBit)
-	{
-		FMemory::Free(Get().AsPtr());
-	}
+	checkf(!Find(Id), TEXT("'%s' already bound"), *Debug.Print(Id));
+	Entries.Emplace(Id, &Binding);
 }
 
-FStructBinding FStructBindingOwner::Get() const
+void FCustomBindings::DropStruct(FStructSchemaId Id)
 {
-	check(*this);
-	return FStructBinding(Handle);
+	for (FEntry& Entry : Entries)
+	{
+		if (Entry.Id == Id)
+		{
+			Entries.RemoveAtSwap(&Entry - Entries.GetData(), EAllowShrinking::No);
+			return;
+		}
+	}
+	
+	checkf(false, TEXT("'%s' unbound"), *Debug.Print(Id));
 }
 
-void FStructBindingOwner::ResetOwned()
+ICustomBinding*	FCustomBindings::Find(FStructSchemaId Id) const
 {
-	check(*this);
-	if (Handle & FStructBinding::SchemaBit)
+	for (const FEntry& Entry : Entries)
 	{
-		FMemory::Free(Get().AsPtr());
+		if (Entry.Id == Id)
+		{
+			return Entry.Binding;
+		}
 	}
-	Handle = 0;
-}
 
-void FStructBindingOwner::TakeOwnership(FStructBinding Binding)
-{
-	check(!*this);
-	Handle = Binding.Handle;
-	check(*this);
+	return nullptr;
 }
 
-FStructBindings::~FStructBindings()
-{}
+////////////////////////////////////////////////////////////////////////////////////////////////
 
-void FStructBindings::BindStruct(FStructSchemaId Id, const ICustomStructBinding& Custom)
-{
-	Bind(Id, FStructBinding(Custom));
-}
+FSchemaBindings::~FSchemaBindings()
+{}
 
 static uint16 CountInnerSchemas(TConstArrayView<FMemberBinding> Members)
 {
@@ -181,11 +180,11 @@ static uint16 CountRanges(TConstArrayView<FMemberBinding> Members)
 	return IntCastChecked<uint16>(Out);
 }
 
-void FStructBindings::BindStruct(FStructSchemaId Id, TConstArrayView<FMemberBinding> Members)
+void FSchemaBindings::BindStruct(FStructSchemaId Id, TConstArrayView<FMemberBinding> Members)
 {
 	// Make header, allocate and copy header
-	FStructSchemaBinding Header = { IntCastChecked<uint16>(Members.Num()), CountInnerSchemas(Members), CountRanges(Members) };
-	FStructSchemaBinding* Schema = new (FMemory::MallocZeroed(Header.CalculateSize())) FStructSchemaBinding {Header};
+	FSchemaBinding Header = { IntCastChecked<uint16>(Members.Num()), CountInnerSchemas(Members), CountRanges(Members) };
+	FSchemaBinding* Schema = new (FMemory::MallocZeroed(Header.CalculateSize())) FSchemaBinding {Header};
 
 	// Write footer
 	FMemberBinder Footer(*Schema);
@@ -207,33 +206,27 @@ void FStructBindings::BindStruct(FStructSchemaId Id, TConstArrayView<FMemberBind
 		}
 	}
 
-	// Register
-	Bind(Id, FStructBinding(*Schema));
-}
-
-void FStructBindings::Bind(FStructSchemaId Id, FStructBinding Binding)
-{
+	// Bind
 	if (Id.Idx >= static_cast<uint32>(Bindings.Num()))
 	{
 		Bindings.SetNum(Id.Idx + 1);
 	}
-
-	Bindings[Id.Idx].TakeOwnership(Binding);
+	checkf(!Bindings[Id.Idx], TEXT("'%s' already bound"), *Debug.Print(Id));
+	Bindings[Id.Idx].Reset(Schema);
 }
 
-FStructBinding FStructBindings::Get(FStructSchemaId Id) const
+const FSchemaBinding& FSchemaBindings::GetStruct(FStructSchemaId Id) const
 {
 	checkf(Id.Idx < (uint32)Bindings.Num() && Bindings[Id.Idx], TEXT("'%s' is unbound"), *Debug.Print(Id));
-	return Bindings[Id.Idx].Get();
+	return *Bindings[Id.Idx].Get();
 }
 
-void FStructBindings::DropStruct(FStructSchemaId Id)
+void FSchemaBindings::DropStruct(FStructSchemaId Id)
 {
 	checkf(Id.Idx < (uint32)Bindings.Num() && Bindings[Id.Idx], TEXT("'%s' is unbound"), *Debug.Print(Id));
-	Bindings[Id.Idx].ResetOwned();
+	Bindings[Id.Idx].Reset();
 }
 
-
 //////////////////////////////////////////////////////////////////////////
 
 uint32 FIdTranslatorBase::CalculateTranslationSize(int32 NumSavedNames, const FSchemaBatch& Batch)
diff --git a/Engine/Source/Programs/PlainPropsTests/Private/PlainPropsLoad.cpp b/Engine/Source/Programs/PlainPropsTests/Private/PlainPropsLoad.cpp
--- a/Engine/Source/Programs/PlainPropsTests/Private/PlainPropsLoad.cpp
+++ b/Engine/Source/Programs/PlainPropsTests/Private/PlainPropsLoad.cpp
@@ -28,14 +28,14 @@ public:
 		check(Memcpy.Offset == AsMemcpy().Offset && Memcpy.Size == AsMemcpy().Size);
 	}
 
-	explicit FLoadStructPlan(const ICustomStructBinding& Custom)
+	explicit FLoadStructPlan(const ICustomBinding& Custom)
 	: Handle(uint64(&Custom) | CustomMask)
 	{
 		check(&Custom == &AsCustom());
 	}
 
 	// @param OffsetWidth Usage unimplemented, store size and offsets as 8/16/32/64-bit 
-	explicit FLoadStructPlan(const FStructSchemaBinding& Schema, ELeafWidth OffsetWidth, bool bSparse)
+	explicit FLoadStructPlan(const FSchemaBinding& Schema, ELeafWidth OffsetWidth, bool bSparse)
 	: Handle(uint64(&Schema) | (uint64(OffsetWidth) << 1) | (uint64(bSparse) << KernelAddressBit) | SchemaMask)
 	{
 		check(&Schema == &AsSchema());
@@ -47,8 +47,8 @@ public:
 	bool						IsMemcpy() const		{ return (Handle & LoMask) == MemcpyMask; }
 	bool						IsCustom() const		{ return (Handle & LoMask) == CustomMask; }
 	FLoadStructMemcpy			AsMemcpy() const		{ check(IsMemcpy()); return { static_cast<uint32>(Handle >> 32), static_cast<uint32>(Handle) >> 1 }; }
-	const ICustomStructBinding&	AsCustom() const		{ check(IsCustom()); return *AsPtr<ICustomStructBinding>(); }
-	const FStructSchemaBinding&	AsSchema() const		{ check(IsSchema()); return *AsPtr<FStructSchemaBinding>(); }
+	const ICustomBinding&		AsCustom() const		{ check(IsCustom()); return *AsPtr<ICustomBinding>(); }
+	const FSchemaBinding&		AsSchema() const		{ check(IsSchema()); return *AsPtr<FSchemaBinding>(); }
 
 private:
 	// This bit is always zero in user mode addresses and most likely won't be used by current or future
@@ -118,7 +118,7 @@ static uint16 CountEnums(const FStructSchema& Schema)
 	return Num;
 }
 
-static bool HasDifferentSupers(const FStructSchema& From, const FStructSchemaBinding& To, TConstArrayView<FStructSchemaId> ToStructIds)
+static bool HasDifferentSupers(const FStructSchema& From, const FSchemaBinding& To, TConstArrayView<FStructSchemaId> ToStructIds)
 {
 	if (From.Inheritance == ESuper::No)
 	{
@@ -182,19 +182,19 @@ static void CopyMemberBinding(/* in-out */ FMemberVisitor& BindIt, /* in-out */
 		case EMemberKind::Leaf:		CopyMemberBinding(BindIt.GrabLeaf(), InnerSchemaIt, Out);	break;
 		case EMemberKind::Range:	CopyMemberBinding(BindIt.GrabRange(), InnerSchemaIt, Out);	break;
 		case EMemberKind::Struct:	CopyMemberBinding(BindIt.GrabStruct(), InnerSchemaIt, Out);	break;
-		default:					check(false);											break;
+		default:					check(false);												break;
 	}
 }
 
-static void CreateSubsetBindingWithoutEnumIds(const FStructSchema& From, const FStructSchemaBinding& To, TConstArrayView<FMemberId> ToNames,  uint16 NumEnums, SubsetByteArray& Out)
+static void CreateSubsetBindingWithoutEnumIds(const FStructSchema& From, const FSchemaBinding& To, TConstArrayView<FMemberId> ToNames,  uint16 NumEnums, SubsetByteArray& Out)
 {
 	check(To.NumMembers == ToNames.Num());
 	check(To.NumMembers >= From.NumMembers);
 
 	int32 OutPos = Out.Num();
-	FStructSchemaBinding Header = { From.NumMembers, From.NumInnerSchemas - NumEnums, From.NumRangeTypes };
+	FSchemaBinding Header = { From.NumMembers, From.NumInnerSchemas - NumEnums, From.NumRangeTypes };
 	Out.AddUninitialized(Header.CalculateSize());
-	FStructSchemaBinding* Schema = new (&Out[OutPos]) FStructSchemaBinding {Header};
+	FSchemaBinding* Schema = new (&Out[OutPos]) FSchemaBinding {Header};
 	
 	FMemberVisitor ToIt(To);
 	FMemberBinder Footer(*Schema);
@@ -211,16 +211,16 @@ static void CreateSubsetBindingWithoutEnumIds(const FStructSchema& From, const F
 	check(InnerSchemaIt == From.GetInnerSchemas() + From.NumInnerSchemas);
 }
 
-static void CloneBindingWithReplacedStructIds(const FSchemaId* FromIds, const FStructSchemaBinding& To, SubsetByteArray& Out)
+static void CloneBindingWithReplacedStructIds(const FSchemaId* FromIds, const FSchemaBinding& To, SubsetByteArray& Out)
 {
 	uint32 Size = To.CalculateSize();
 	Out.AddUninitialized(Size);
-	FStructSchemaBinding* Schema = reinterpret_cast<FStructSchemaBinding*>(&Out[Out.Num() - Size]);
+	FSchemaBinding* Schema = reinterpret_cast<FSchemaBinding*>(&Out[Out.Num() - Size]);
 	FMemory::Memcpy(Schema, &To, Size);
 	FMemory::Memcpy(const_cast<FSchemaId*>(Schema->GetInnerSchemas()), FromIds, To.NumInnerSchemas * sizeof(FSchemaId));
 }
 
-[[nodiscard]] static FLoadStructPlan MakeSchemaLoadPlan(const FStructSchema& From, const FStructSchemaBinding& To, TConstArrayView<FMemberId> ToMemberIds, TConstArrayView<FStructSchemaId> ToStructIds, SubsetByteArray& OutSubsetSchemas)
+[[nodiscard]] static FLoadStructPlan MakeSchemaLoadPlan(const FStructSchema& From, const FSchemaBinding& To, TConstArrayView<FMemberId> ToMemberIds, TConstArrayView<FStructSchemaId> ToStructIds, SubsetByteArray& OutSubsetSchemas)
 {
 	uint16 NumEnums = CountEnums(From);
 	if (From.NumMembers < To.NumMembers || NumEnums || HasDifferentSupers(From, To, ToStructIds))
@@ -244,19 +244,19 @@ static void CloneBindingWithReplacedStructIds(const FSchemaId* FromIds, const FS
 	return FLoadStructPlan(To, ELeafWidth::B32, !From.IsDense);
 }
 
-[[nodiscard]] static TOptional<FLoadStructMemcpy> TryMakeMemcpyPlan(const FStructSchema& From, const FStructSchemaBinding& To, TConstArrayView<FMemberId> ToMemberIds)
+[[nodiscard]] static TOptional<FLoadStructMemcpy> TryMakeMemcpyPlan(const FStructSchema& From, const FSchemaBinding& To, TConstArrayView<FMemberId> ToMemberIds)
 {
 	// todo
 	return NullOpt;
 }
 
-[[nodiscard]] static FLoadStructPlan MakeLoadPlan(const FStructSchema& From, const FStructSchemaBinding& To, TConstArrayView<FMemberId> ToMemberIds, TConstArrayView<FStructSchemaId> ToStructIds, SubsetByteArray& OutSubsetSchemas)
+[[nodiscard]] static FLoadStructPlan MakeLoadPlan(const FStructSchema& From, const FSchemaBinding& To, TConstArrayView<FMemberId> ToMemberIds, TConstArrayView<FStructSchemaId> ToStructIds, SubsetByteArray& OutSubsetSchemas)
 {
 	TOptional<FLoadStructMemcpy> Memcpy = TryMakeMemcpyPlan(From, To, ToMemberIds);
 	return Memcpy ? FLoadStructPlan(Memcpy.GetValue()) : MakeSchemaLoadPlan(From, To, ToMemberIds, ToStructIds, OutSubsetSchemas);
 }
 
-FLoadBatch* CreateLoadPlans(FReadBatchId ReadId, const FDeclarations& Declarations, const FStructBindings& Bindings, TConstArrayView<FStructSchemaId> RuntimeIds)
+FLoadBatch* CreateLoadPlans(FReadBatchId ReadId, const FDeclarations& Declarations, const FCustomBindings& Customs, const FSchemaBindings& Schemas, TConstArrayView<FStructSchemaId> RuntimeIds)
 {
 	check(NumStructSchemas(ReadId) == RuntimeIds.Num());
 
@@ -272,16 +272,15 @@ FLoadBatch* CreateLoadPlans(FReadBatchId ReadId, const FDeclarations& Declaratio
 	for (FStructSchemaId SavedId = { 0 }; SavedId.Idx < NumPlans; ++SavedId.Idx)
 	{
 		FStructSchemaId RuntimeId = RuntimeIds[SavedId.Idx];
-		FStructBinding Binding = Bindings.Get(RuntimeId);
 		int32 SubsetSchemaOffset = SubsetSchemaData.Num();
-		if (Binding.IsCustom())
+		if (const ICustomBinding* Custom = Customs.FindStruct(RuntimeId))
 		{
-			Plans[SavedId.Idx] = FLoadStructPlan(Binding.AsCustom()) ;
+			Plans[SavedId.Idx] = FLoadStructPlan(*Custom) ;
 		}
 		else
 		{
 			const FStructSchema& From = ResolveStructSchema(ReadId, SavedId);
-			const FStructSchemaBinding& To = Binding.AsSchema();
+			const FSchemaBinding& To = Schemas.GetStruct(RuntimeId);
 			// Possible optimization - some simple memcpy cases doesn't need to resolve the declaration
 			TConstArrayView<FMemberId> ToMemberIds = Declarations.Get(RuntimeId).GetMemberOrder();
 			Plans[SavedId.Idx] = MakeLoadPlan(From, To, ToMemberIds, RuntimeIds, /* out */ SubsetSchemaData);	
@@ -306,10 +305,10 @@ FLoadBatch* CreateLoadPlans(FReadBatchId ReadId, const FDeclarations& Declaratio
 		{
 			if (int32 Size = SubsetSchemaSizes[Idx])
 			{
-				check(IsAligned(Size, alignof(FStructSchemaBinding)));
+				check(IsAligned(Size, alignof(FSchemaBinding)));
 				check(Plans[Idx].IsSchema());
 				bool bSparse = Plans[Idx].IsSparseSchema();
-				Out->Plans[Idx] = FLoadStructPlan(*reinterpret_cast<const FStructSchemaBinding*>(It), ELeafWidth::B32, bSparse);
+				Out->Plans[Idx] = FLoadStructPlan(*reinterpret_cast<const FSchemaBinding*>(It), ELeafWidth::B32, bSparse);
 				It += Size;
 			}
 		}
@@ -488,7 +487,7 @@ template<bool bSparse, typename OffsetType>
 class TMemberLoader : public FRangeLoader
 {
 public:
-	TMemberLoader(FByteReader Values, const FStructSchemaBinding& Schema, const FLoadBatch& InBatch)
+	TMemberLoader(FByteReader Values, const FSchemaBinding& Schema, const FLoadBatch& InBatch)
 	: Types(Schema.Members, Schema.NumMembers)
 	, Offsets(Schema.GetOffsets())
 	, InnerStructSchemas(static_cast<const FStructSchemaId*>(Schema.GetInnerSchemas()), Schema.NumInnerSchemas)
@@ -634,7 +633,7 @@ void LoadStruct(uint8* Dst, FStructView Src, const FLoadBatch& Batch)
 void ConstructAndLoadStruct(uint8* Dst, FByteReader Src, FStructSchemaId Id, const FLoadBatch& Batch)
 {
 	FLoadStructPlan Plan = Batch[Id];
-	checkf(!Plan.IsSchema(), TEXT("Non-default constructible types requires ICustomStructBinding or in rare cases memcpying"));
+	checkf(!Plan.IsSchema(), TEXT("Non-default constructible types requires ICustomBinding or in rare cases memcpying"));
 
 	if (Plan.IsMemcpy())
 	{
diff --git a/Engine/Source/Programs/PlainPropsTests/Private/PlainPropsSave.cpp b/Engine/Source/Programs/PlainPropsTests/Private/PlainPropsSave.cpp
--- a/Engine/Source/Programs/PlainPropsTests/Private/PlainPropsSave.cpp
+++ b/Engine/Source/Programs/PlainPropsTests/Private/PlainPropsSave.cpp
@@ -323,7 +323,6 @@ template<class SaverType, typename InnerContextType>
 	return { FMemberType(Member.RangeBindings[0].GetSizeType()), Member.InnermostSchema, ToMemberTypes(Member.InnerTypes) };
 }
 
-
 static void SaveMember(FMemberBuilder& Out, const uint8* Struct, FMemberId Name, const FSaveContext& Ctx, FLeafMemberBinding Member)
 {
 	FUnpackedLeafType Type = { ToLeafType(Member.Leaf.Type), Member.Leaf.Width };
@@ -342,15 +341,14 @@ static void SaveMember(FMemberBuilder& Out, const uint8* Struct, FMemberId Name,
 
 TUniquePtr<FBuiltStruct> SaveStruct(const uint8* Struct, FStructSchemaId Id, const FSaveContext& Ctx)
 {
-	FStructBinding Binding = Ctx.Bindings.Get(Id);
 	const FStructDeclaration& Declaration = Ctx.Declarations.Get(Id);
-	FMemberBuilder Out;
 
-	if (Binding.IsCustom())
+	FMemberBuilder Out;
+	if (ICustomBinding* Custom = Ctx.Customs.FindStruct(Id))
 	{
-		Binding.AsCustom().SaveStruct(Out, Struct, nullptr, Ctx.Debug);
+		Custom->SaveStruct(Out, Struct, nullptr, Ctx.Declarations.GetDebug());
 	}
-	else for (FMemberVisitor It(Binding.AsSchema()); It.HasMore(); )
+	else for (FMemberVisitor It(Ctx.Schemas.GetStruct(Id)); It.HasMore(); )
 	{
 		FMemberId Name = Declaration.GetMemberOrder()[It.GetIndex()];
 		switch (It.PeekKind())
@@ -361,7 +359,7 @@ TUniquePtr<FBuiltStruct> SaveStruct(const uint8* Struct, FStructSchemaId Id, con
 		}
 	}
 	
-	return Out.BuildAndReset(Declaration, Ctx.Debug);
+	return Out.BuildAndReset(Declaration, Ctx.Declarations.GetDebug());
 }
 
 } // namespace PlainProps
\ No newline at end of file
diff --git a/Engine/Source/Programs/PlainPropsTests/Public/PlainPropsBind.h b/Engine/Source/Programs/PlainPropsTests/Public/PlainPropsBind.h
--- a/Engine/Source/Programs/PlainPropsTests/Public/PlainPropsBind.h
+++ b/Engine/Source/Programs/PlainPropsTests/Public/PlainPropsBind.h
@@ -93,23 +93,22 @@ union FMemberBindType
 	constexpr explicit FMemberBindType(ERangeSizeType MaxSize) : Range({EMemberKind::Range, MaxSize}) {}
 	constexpr explicit FMemberBindType(FStructType In) : Struct(In) {}
 	
-	bool				IsLeaf() const			{ return Kind == EMemberKind::Leaf; }
-	bool				IsRange() const			{ return Kind == EMemberKind::Range; }
-	bool				IsStruct() const		{ return Kind == EMemberKind::Struct; }
-	EMemberKind			GetKind() const			{ return Kind; }
+	bool					IsLeaf() const		{ return Kind == EMemberKind::Leaf; }
+	bool					IsRange() const		{ return Kind == EMemberKind::Range; }
+	bool					IsStruct() const	{ return Kind == EMemberKind::Struct; }
+	EMemberKind				GetKind() const		{ return Kind; }
 	
-	FLeafBindType		AsLeaf() const			{ check(IsLeaf());		return Leaf; }
-	FRangeBindType		AsRange() const			{ check(IsRange());		return Range; }
-	FStructBindType		AsStruct() const		{ check(IsStruct());	return Struct; }
-	uint8				AsByte() const			{ return BitCast<uint8>(*this); }
+	FLeafBindType			AsLeaf() const		{ check(IsLeaf());		return Leaf; }
+	FRangeBindType			AsRange() const		{ check(IsRange());		return Range; }
+	FStructBindType			AsStruct() const	{ check(IsStruct());	return Struct; }
+	uint8					AsByte() const		{ return BitCast<uint8>(*this); }
 
 	friend inline bool operator==(FMemberBindType A, FMemberBindType B) { return A.AsByte() == B.AsByte(); }
-
 private:
-    EMemberKind			Kind : 2;
-	FLeafBindType		Leaf;
-    FRangeBindType		Range;
-    FStructBindType		Struct;
+    EMemberKind				Kind : 2;
+	FLeafBindType			Leaf;
+    FRangeBindType			Range;
+    FStructBindType			Struct;
 };
 
 static_assert(sizeof(FMemberBindType) == 1);
@@ -117,7 +116,7 @@ static_assert(sizeof(FMemberBindType) == 1);
 ////////////////////////////////////////////////////////////////////////////////////////////////
 
 // Members are loaded in saved FStructSchema order, not current offset order unless upgrade layer reorders
-struct FStructSchemaBinding
+struct FSchemaBinding
 {
 	uint16					NumMembers;
 	uint16					NumInnerSchemas;
@@ -139,13 +138,13 @@ struct FUnpackedLeafBindType
 	ELeafBindType			Type;
 	union
 	{
-		ELeafWidth				Width;
-		uint8					BitfieldIdx;
+		ELeafWidth			Width;
+		uint8				BitfieldIdx;
 	};
 	
 	//constexpr FUnpackedLeafBindType(ELeafBindType InType, ELeafWidth InWidth) : Type(InType), Width(InWidth) {}
 	constexpr FUnpackedLeafBindType(FLeafBindType In)
-		: Type(In.Bind.Type)
+	: Type(In.Bind.Type)
 	{
 		if (Type == ELeafBindType::BitfieldBool)
 		{
@@ -191,7 +190,8 @@ struct FStructMemberBinding
 class FMemberVisitor
 {
 public:
-	explicit FMemberVisitor(const FStructSchemaBinding& InSchema);
+	explicit FMemberVisitor(const FSchemaBinding& InSchema);
+
 	bool						HasMore() const			{ return MemberIdx < NumMembers; }
 	uint16						GetIndex() const		{ return MemberIdx; }
 	
@@ -204,7 +204,7 @@ public:
 	void						SkipMember();
 
 protected: // for unit tests
-	const FStructSchemaBinding& Schema;
+	const FSchemaBinding&		Schema;
 	const uint16				NumMembers;
 	uint16						MemberIdx = 0;
 	uint16						InnerRangeIdx = 0;		// Types of [nested] ranges
@@ -224,44 +224,37 @@ protected: // for unit tests
 
 enum ECustomLoadMethod { Construct, Assign };
 
-class ICustomStructBinding
+class ICustomBinding
 {
 public:
-	virtual void			LoadStruct(void* Dst, FStructView Src, ECustomLoadMethod Method, const FLoadBatch& Batch) const = 0;
-	virtual void			SaveStruct(FMemberBuilder& Dst, const void* Src, void* UserData, const FDebugIds& Debug) const = 0;
+	virtual ~ICustomBinding() {}
+	virtual void				SaveStruct(FMemberBuilder& Dst, const void* Src, void* UserData, const FDebugIds& Debug) = 0;
+	virtual void				LoadStruct(void* Dst, FStructView Src, ECustomLoadMethod Method, const FLoadBatch& Batch) const = 0;
 };
 
-class FStructBinding
+class FCustomBindings
 {
 public:
-	explicit FStructBinding(const FStructSchemaBinding& Schema) : Handle(uint64(&Schema) | SchemaBit) {}
-	explicit FStructBinding(const ICustomStructBinding& Custom) : Handle(uint64(&Custom)) {}
-	
-	bool						IsSchema() const		{ return Handle & SchemaBit; }
-	bool						IsCustom() const		{ return !IsSchema(); }
-	const FStructSchemaBinding&	AsSchema() const		{ check(IsSchema()); return *static_cast<FStructSchemaBinding*>(AsPtr()); }
-	const ICustomStructBinding&	AsCustom() const		{ check(IsCustom()); return *static_cast<ICustomStructBinding*>(AsPtr()); }
+	UE_NONCOPYABLE(FCustomBindings);
+	explicit FCustomBindings(const FDebugIds& In) : Debug(In) {}
+
+	// @param Binding must outlive this or call DropStruct()
+	void						BindStruct(FStructSchemaId Id, ICustomBinding& Binding);
+	ICustomBinding*				FindStruct(FStructSchemaId Id)								{ return Find(Id); }
+	const ICustomBinding*		FindStruct(FStructSchemaId Id) const						{ return Find(Id); }
+	void						DropStruct(FStructSchemaId Id);
 
 private:
-	static constexpr uint64 SchemaBit = 1;
-	uint64 Handle;
+	ICustomBinding*				Find(FStructSchemaId Id) const;
 
-	friend class FStructBindingOwner;
-	explicit FStructBinding(uint64 InHandle) : Handle(InHandle) {}
-	void*						AsPtr() const			{ return reinterpret_cast<void*>(Handle & ~SchemaBit); }
-};
+	struct FEntry
+	{
+		FStructSchemaId	Id;
+		ICustomBinding* Binding;
+	};
 
-class FStructBindingOwner
-{
-	uint64 Handle = 0;
-public:
-	UE_NONCOPYABLE(FStructBindingOwner);
-	FStructBindingOwner() = default;
-	~FStructBindingOwner();
-	explicit					operator bool() const	{ return Handle != 0; }
-	FStructBinding				Get() const;
-	void						ResetOwned();
-	void						TakeOwnership(FStructBinding Binding);
+	TArray<FEntry, TInlineAllocator<8>>		Entries;
+	const FDebugIds&						Debug;
 };
 
 ////////////////////////////////////////////////////////////////////////////////////////////////
@@ -434,23 +427,20 @@ struct FMemberBinding
 
 ////////////////////////////////////////////////////////////////////////////////////////////////
 
-class FStructBindings
+class FSchemaBindings
 {
 public:
-	UE_NONCOPYABLE(FStructBindings);
-	explicit FStructBindings(const FDebugIds& In) : Debug(In) {}
-	~FStructBindings();
+	UE_NONCOPYABLE(FSchemaBindings);
+	explicit FSchemaBindings(const FDebugIds& In) : Debug(In) {}
+	~FSchemaBindings();
 
-	void							BindStruct(FStructSchemaId Id, const ICustomStructBinding& Custom);
-	void							BindStruct(FStructSchemaId Id, TConstArrayView<FMemberBinding> Schema);
-	FStructBinding					Get(FStructSchemaId Id) const;
-	void							DropStruct(FStructSchemaId Id);
+	void								BindStruct(FStructSchemaId Id, TConstArrayView<FMemberBinding> Schema);
+	const FSchemaBinding&				GetStruct(FStructSchemaId Id) const;
+	void								DropStruct(FStructSchemaId Id);
 
 private:
-	TArray<FStructBindingOwner>		Bindings;
-	const FDebugIds&				Debug;
-
-	void							Bind(FStructSchemaId Id, FStructBinding Binding);
+	TArray<TUniquePtr<FSchemaBinding>>	Bindings;
+	const FDebugIds&					Debug;
 };
 
 ////////////////////////////////////////////////////////////////////////////////////////////////
@@ -520,11 +510,15 @@ FStructSchemaId BindCustomStructOnce()
 template<class Type, class CustomBinding, class Runtime>
 FMemberBindType BindMemberStruct(FOptionalSchemaId& OutSchema)
 {
-	FStructSchemaId Id	= std::is_void_v<CustomBinding>
-						? IndexNativeStruct<Type, typename Runtime::Ids>()
-						: BindCustomStructOnce<CustomBinding, Runtime>();
+	if constexpr (std::is_void_v<CustomBinding>)
+	{
+		OutSchema = FOptionalSchemaId(IndexNativeStruct<Type, typename Runtime::Ids>());
+	}
+	else
+	{
+		OutSchema = FOptionalSchemaId(BindCustomStructOnce<CustomBinding, Runtime>());
+	}
 
-	OutSchema = FOptionalSchemaId(Id);
 	return FMemberBindType(FStructType{EMemberKind::Struct, /* IsDynamic */ 0, /* IsSuper */ 0});
 }
 
@@ -662,28 +656,6 @@ FEnumSchemaId DeclareNativeEnum(FDeclarations& Out, EEnumMode Mode)
 	return Id;
 }
 
-//
-//template<class Ctti, class Rttis>
-//FStructSchemaId BindNativeStruct(FDeclarations& Runtime, EMemberPresence Occupancy)
-//{
-//	using Ids = typename Rttis::IdsType;
-//	FTypeId Type = Ids::IndexNativeType(Ctti::Name);
-//
-//	FMemberBinding MemberBindings[Ctti::NumVars];
-//	FMemberId MemberIds[Ctti::NumVars];
-//	ForEachVar<Ctti>([&]<class Var>()
-//	{ 
-//		MemberIds[Var::Index] = Ids::IndexMember(Var::Name);
-//		MemberBindings[Var::Index] = BindMember<Var, Rttis>();
-//	});
-//
-//	FOptionalSchemaId SuperId = !std::is_void_v<Ctti::Super> ? IndexNativeStruct<Ctti::Super, Ids>() : NoId;
-//	FStructSchemaId Id = Runtime.DeclareStruct(Type, SuperId, MemberIds, Occupancy);
-//	Runtime.BindStruct(Id, MemberBindings);
-//	return Id;
-//}
-
-
 template<class Ctti, class Ids>
 FStructSchemaId DeclareNativeStruct(FDeclarations& Out, EMemberPresence Occupancy)
 {
@@ -701,7 +673,7 @@ FStructSchemaId DeclareNativeStruct(FDeclarations& Out, EMemberPresence Occupanc
 }
 
 template<class Ctti, class Runtime>
-void BindNativeStruct(FStructBindings& Out, FStructSchemaId DeclaredId)
+void BindNativeStruct(FSchemaBindings& Out, FStructSchemaId DeclaredId)
 {
 	FMemberBinding MemberBindings[Ctti::NumVars];
 	ForEachVar<Ctti>([&]<class Var>()
@@ -715,7 +687,7 @@ void BindNativeStruct(FStructBindings& Out, FStructSchemaId DeclaredId)
 
 struct FMemberBinder
 {
-	FMemberBinder(FStructSchemaBinding& InSchema)
+	FMemberBinder(FSchemaBinding& InSchema)
 	: Schema(InSchema)
 	, MemberIt(Schema.Members)
 	, RangeTypeIt(const_cast<FMemberBindType*>(Schema.GetInnerRangeTypes()))
@@ -758,7 +730,7 @@ struct FMemberBinder
 		*InnerSchemaIt++ = InnermostSchema;
 	}
 
-	FStructSchemaBinding& Schema;
+	FSchemaBinding& Schema;
 	FMemberBindType* MemberIt;
 	FMemberBindType* RangeTypeIt;
 	uint32* OffsetIt;
diff --git a/Engine/Source/Programs/PlainPropsTests/Public/PlainPropsDeclare.h b/Engine/Source/Programs/PlainPropsTests/Public/PlainPropsDeclare.h
--- a/Engine/Source/Programs/PlainPropsTests/Public/PlainPropsDeclare.h
+++ b/Engine/Source/Programs/PlainPropsTests/Public/PlainPropsDeclare.h
@@ -63,8 +63,9 @@ public:
 	const FEnumDeclaration&							Get(FEnumSchemaId Id) const		{ Check(Id); return *DeclaredEnums[Id.Idx]; }
 	const FStructDeclaration&						Get(FStructSchemaId Id) const	{ Check(Id); return *DeclaredStructs[Id.Idx]; }
 	
-	TConstArrayView<TUniquePtr<FEnumDeclaration>>	GetEnums() const	{ return DeclaredEnums; }
-	TConstArrayView<TUniquePtr<FStructDeclaration>>	GetStructs() const	{ return DeclaredStructs; }
+	TConstArrayView<TUniquePtr<FEnumDeclaration>>	GetEnums() const				{ return DeclaredEnums; }
+	TConstArrayView<TUniquePtr<FStructDeclaration>>	GetStructs() const				{ return DeclaredStructs; }
+	const FDebugIds&								GetDebug() const				{ return Debug; }
 
 protected:
 	TArray<TUniquePtr<FEnumDeclaration>>			DeclaredEnums;
diff --git a/Engine/Source/Programs/PlainPropsTests/Public/PlainPropsLoad.h b/Engine/Source/Programs/PlainPropsTests/Public/PlainPropsLoad.h
--- a/Engine/Source/Programs/PlainPropsTests/Public/PlainPropsLoad.h
+++ b/Engine/Source/Programs/PlainPropsTests/Public/PlainPropsLoad.h
@@ -8,15 +8,16 @@
 
 namespace PlainProps 
 {
-
+	
+class FCustomBindings;
 class FDeclarations;
 struct FLoadBatch;
 class FRangeBinding;
-class FStructBindings;
+class FSchemaBindings;
 struct FTypedRange;
 
 
-PP_API FLoadBatch* CreateLoadPlans(FReadBatchId ReadId, const FDeclarations& Declarations, const FStructBindings& Bindings, TConstArrayView<FStructSchemaId> RuntimeIds);
+PP_API FLoadBatch* CreateLoadPlans(FReadBatchId ReadId, const FDeclarations& Declarations, const FCustomBindings& Customs, const FSchemaBindings& Schemas, TConstArrayView<FStructSchemaId> RuntimeIds);
 PP_API void DestroyLoadPlans(FLoadBatch* Batch);
 PP_API void LoadStruct(uint8* Dst, FByteReader Src, FStructSchemaId Id, const FLoadBatch& Batch);
 PP_API void LoadStruct(uint8* Dst, FStructView Src, const FLoadBatch& Batch);
diff --git a/Engine/Source/Programs/PlainPropsTests/Public/PlainPropsSave.h b/Engine/Source/Programs/PlainPropsTests/Public/PlainPropsSave.h
--- a/Engine/Source/Programs/PlainPropsTests/Public/PlainPropsSave.h
+++ b/Engine/Source/Programs/PlainPropsTests/Public/PlainPropsSave.h
@@ -9,14 +9,15 @@ namespace PlainProps
 {
 
 struct FBuiltStruct;
+class FCustomBindings;
 class FDeclarations;
-class FStructBindings;
+class FSchemaBindings;
 
 struct FSaveContext
 {
-	const FDeclarations& Declarations;
-	const FStructBindings& Bindings;
-	const FDebugIds& Debug;
+	const FDeclarations&		Declarations;
+	const FSchemaBindings&		Schemas;
+	FCustomBindings&			Customs;
 };
 
 [[nodiscard]] TUniquePtr<FBuiltStruct> SaveStruct(const uint8* Struct, FStructSchemaId Id, const FSaveContext& Context);
diff --git a/Engine/Source/Programs/PlainPropsTests/Public/PlainPropsUeCoreBindings.h b/Engine/Source/Programs/PlainPropsTests/Public/PlainPropsUeCoreBindings.h
--- a/Engine/Source/Programs/PlainPropsTests/Public/PlainPropsUeCoreBindings.h
+++ b/Engine/Source/Programs/PlainPropsTests/Public/PlainPropsUeCoreBindings.h
@@ -15,6 +15,10 @@ namespace PlainProps
 	void AppendString(FString& Out, const FName& Name) { Name.AppendString(Out); }
 }
 
+// Todo: PP_NAME_STRUCT(, FName);
+struct FName_Ctti { static constexpr char Name[] = "FName"; };
+FName_Ctti CttiOfPtr(FName*);
+
 namespace PlainProps::UE
 {
 
@@ -31,7 +35,7 @@ namespace PlainProps::UE
 
 	//template<typename Ctti>
 	//FStructSchemaId			BindStructInterlaced(TConstArrayView<FMemberBinding> NonCttiMembers);
-	//FStructSchemaId			BindStruct(FStructSchemaId Id, const ICustomStructBinding& Custom);
+	//FStructSchemaId			BindStruct(FStructSchemaId Id, const ICustomBinding& Custom);
 	//FStructSchemaId			BindStruct(FTypeId Type, FOptionalSchemaId Super, TConstArrayView<FNamedMemberBinding> Members, EMemberPresence Occupancy);
 	//void					DropStruct(FStructSchemaId Id) { Types.DropStruct(Id); }
 
@@ -301,7 +305,7 @@ struct FSetOps
 };
 
 template <class Ids, typename T>
-struct TSetDeltaBinding : public ICustomStructBinding
+struct TSetDeltaBinding : public ICustomBinding
 {
 	using Type = TSet<T>;
 	static constexpr EMemberPresence Occupancy = EMemberPresence::AllowSparse;
@@ -422,7 +426,7 @@ struct TSetDeltaBinding : public ICustomStructBinding
 //	};
 //
 //template <typename... Ts>
-//struct TVariantBinding : public ICustomStructBinding
+//struct TVariantBinding : public ICustomBinding
 //{
 //	using VariantType = TVariant<Ts...>;
 //
diff --git a/Engine/Source/Programs/PlainPropsTests/Tests/PlainPropsUeCoreTests.cpp b/Engine/Source/Programs/PlainPropsTests/Tests/PlainPropsUeCoreTests.cpp
--- a/Engine/Source/Programs/PlainPropsTests/Tests/PlainPropsUeCoreTests.cpp
+++ b/Engine/Source/Programs/PlainPropsTests/Tests/PlainPropsUeCoreTests.cpp
@@ -23,7 +23,7 @@ namespace PlainProps::UE::Test
 
 static TIdIndexer<FName>	GNames;
 static FDeclarations		GTypes(/* debug */ GNames);
-static FStructBindings		GBindings(/* debug */ GNames);
+static FSchemaBindings		GSchemas(/* debug */ GNames);
 
 struct FIds
 {
@@ -55,67 +55,130 @@ struct FDefaultRuntime
 	template<class T> using CustomBindings = TCustomBindings<T>;
 
 	static FDeclarations&			GetDeclarations()				{ return GTypes; }
-	static FStructBindings&			GetBindings()					{ return GBindings; }
+	static FSchemaBindings&			GetSchemas()					{ return GSchemas; }
 };
 
 ////////////////////////////////////////////////////////////////////////////////////////////////
 
-template<class T, class Runtime = FDefaultRuntime>
-struct TScopedStructBinding
+template<typename Enum, EEnumMode Mode, class Runtime = FDefaultRuntime>
+struct TScopedEnumDeclaration
+{
+	using Ids = typename Runtime::Ids;
+	using Ctti = CttiOf<Enum>;
+
+	FEnumSchemaId Id;
+	TScopedEnumDeclaration() : Id(DeclareNativeEnum<Ctti, Ids>(Runtime::GetDeclarations(), Mode)) {}
+	~TScopedEnumDeclaration() { Runtime::GetDeclarations().DropEnum(Id); }
+};
+
+template<class T, EMemberPresence Occupancy = EMemberPresence::AllowSparse, class Runtime = FDefaultRuntime>
+struct TScopedStructDeclaration
 {
 	using Ids = typename Runtime::Ids;
-	using Ctti = CttiOf<T>;
 
-	TScopedStructBinding(EMemberPresence Occupancy = EMemberPresence::AllowSparse)
-	: Id(DeclareNativeStruct<Ctti, Ids>(Runtime::GetDeclarations(), Occupancy))
+	FStructSchemaId Id;
+
+	TScopedStructDeclaration()
+	: Id(DeclareNativeStruct<CttiOf<T>, Ids>(Runtime::GetDeclarations(), Occupancy))
+	{}
+
+	~TScopedStructDeclaration()
 	{
-		BindNativeStruct<Ctti, Runtime>(Runtime::GetBindings(), Id);
+		Runtime::GetDeclarations().DropStruct(Id);
+	}
+
+	const FStructDeclaration& Get() const
+	{
+		return Runtime::GetDeclarations().Get(Id);
+	}
+};
+
+template<class T, EMemberPresence Occupancy = EMemberPresence::AllowSparse, class Runtime = FDefaultRuntime>
+struct TScopedStructBinding : TScopedStructDeclaration<T, Occupancy, Runtime>
+{
+	using TScopedStructDeclaration<T, Occupancy, Runtime>::Id;
+
+	TScopedStructBinding()
+	{
+		BindNativeStruct<CttiOf<T>, Runtime>(Runtime::GetSchemas(), Id);
 	}
 
 	~TScopedStructBinding()
 	{
-		Runtime::GetBindings().DropStruct(Id);
-		Runtime::GetDeclarations().DropStruct(Id);
+		Runtime::GetSchemas().DropStruct(Id);
+	}
+};
+
+////////////////////////////////////////////////////////////////////////////////////////////////
+
+// Maybe replace with macro, e.g. PP_DECLARE_CUSTOM_DENSE_STRUCT(FIds, FName, void, Idx)
+struct FNameDeclaration 
+{
+	FStructSchemaId		Id;
+	FMemberId			Idx;
+
+	FNameDeclaration(FTypeId Type = FIds::IndexNativeType(CttiOf<FName>::Name))
+	: Id(FIds::IndexStruct(Type))
+	, Idx(FIds::IndexMember("Idx"))
+	{
+		GTypes.DeclareStruct(Id, Type, MakeArrayView(&Idx, 1), EMemberPresence::RequireAll);
 	}
 
-	FStructSchemaId Id;
+	~FNameDeclaration()
+	{
+		GTypes.DropStruct(Id);
+	}
 };
 
-template<typename Enum, EEnumMode Mode, class Runtime = FDefaultRuntime>
-struct TScopedEnumBinding
+struct FNameBinding : public ICustomBinding
 {
-	using Ids = typename Runtime::Ids;
-	using Ctti = CttiOf<Enum>;
+	virtual void SaveStruct(FMemberBuilder& Dst, const void* Src, void*, const FDebugIds& Debug) override
+	{
+		FSetElementId Idx = Names.Add(*static_cast<const FName*>(Src));
+		Dst.Add(Declaration.Idx, Idx.AsInteger());
+	}
 
-	FEnumSchemaId Id;
-	TScopedEnumBinding() : Id(DeclareNativeEnum<Ctti, Ids>(Runtime::GetDeclarations(), Mode)) {}
-	~TScopedEnumBinding() { Runtime::GetDeclarations().DropEnum(Id); }
+	virtual void LoadStruct(void* Dst, FStructView Src, ECustomLoadMethod, const FLoadBatch&) const override
+	{
+		FSetElementId Idx = FSetElementId::FromInteger(FMemberReader(Src).GrabLeaf().AsS32());
+		*static_cast<FName*>(Dst) = Names.Get(Idx);
+	}
+
+	FNameDeclaration	Declaration;
+	TSet<FName>			Names;
 };
 
 ////////////////////////////////////////////////////////////////////////////////////////////////
 
-inline constexpr uint32 Magics[] = { 0xFEEDF00D, 0xABCD1234, 0xDADADAAA, 0x99887766, 0xF0F1F2F3 };
+inline constexpr uint32 Magics[] = { 0xFEEDF00D, 0xABCD1234, 0xDADADAAA, 0x99887766, 0xF0F1F2F3 , 0x00112233};
 
 class FBatchSaver
 {
 public:
-	explicit FBatchSaver(const FStructBindings& InBindings) : Bindings(InBindings) {}
+	FBatchSaver();
 
 	template<class T>
-	void						Save(T&& Object) { Save(IndexNativeStruct<T, FIds>(), &Object); }
+	void						Save(T&& Object);
 	TArray64<uint8>				Write() const;
 
 private:
-	void						Save(FStructSchemaId Id, const void* Object); 
-
 	using IdBuiltStructPair = TPair<FStructSchemaId, TUniquePtr<FBuiltStruct>>;
 	TArray<IdBuiltStructPair>	SavedObjects;
-	const FStructBindings&		Bindings;
+	FNameBinding				SavedNames;
+	FCustomBindings				Customs;
 };
 
-void FBatchSaver::Save(FStructSchemaId Id, const void* Object)
+FBatchSaver::FBatchSaver()
+: Customs(/* debug */ GNames)
+{
+	Customs.BindStruct(SavedNames.Declaration.Id, SavedNames);
+}
+
+template<class T>
+void FBatchSaver::Save(T&& Object) 
 {
-	SavedObjects.Emplace(Id, SaveStruct(reinterpret_cast<const uint8*>(Object), Id, {GTypes, GBindings, /* debug */ GNames}));
+	FStructSchemaId Id = IndexNativeStruct<T, FIds>();
+	SavedObjects.Emplace(Id, SaveStruct(reinterpret_cast<const uint8*>(&Object), Id, {GTypes, GSchemas, Customs}));
 }
 
 template<typename ArrayType>
@@ -180,6 +243,10 @@ TArray64<uint8> FBatchSaver::Write() const
 	// Write object terminator
 	WriteSkippableSlice(Out, TConstArrayView64<uint8>());
 	WriteU32(Out, Magics[4]);
+	
+	// Write names
+	WriteNumAndArray(Out, SavedNames.Names.Array());
+	WriteU32(Out, Magics[5]);
 		
 	return Out;
 }
@@ -189,13 +256,14 @@ TArray64<uint8> FBatchSaver::Write() const
 class FBatchLoader
 {
 public:
-	FBatchLoader(const FStructBindings& Bindings, FMemoryView Data)
+	FBatchLoader(FMemoryView Data)
+	: Customs(/* debug */ GNames)
 	{
 		// Read ids
 		FByteReader It(Data);
 		CHECK(It.Grab<uint32>() == Magics[0]);
-		SavedNames = GrabNumAndArray<FName>(It);
-		CHECK(SavedNames.Num() != 0);
+		Ids = GrabNumAndArray<FName>(It);
+		CHECK(Ids.Num() != 0);
 		
 		// Read schemas
 		CHECK(It.Grab<uint32>() == Magics[1]);
@@ -205,7 +273,7 @@ public:
 		CHECK(It.Grab<uint32>() == Magics[2]);
 		
 		// Bind saved ids to runtime ids, make new schemas with new ids and mount them
-		FIdTranslator RuntimeIds(GNames, SavedNames, *SavedSchemas);
+		FIdTranslator RuntimeIds(GNames, Ids, *SavedSchemas);
 		FSchemaBatch* LoadSchemas = CreateTranslatedSchemas(*SavedSchemas, RuntimeIds.Translation);
 		FReadBatchId Batch = MountReadSchemas(LoadSchemas);
 
@@ -221,9 +289,14 @@ public:
 		CHECK(It.Grab<uint32>() == Magics[4]);
 		CHECK(!Objects.IsEmpty());
 
+		// Read names and bind custom loader
+		Names.Names.Append(GrabNumAndArray<FName>(It));
+		Customs.BindStruct(Names.Declaration.Id, Names);
+		CHECK(It.Grab<uint32>() == Magics[5]);
+
 		// Finally create load plans
 		TConstArrayView<FStructSchemaId> LoadStructIds = RuntimeIds.Translation.GetStructIds(SavedSchemas->NumStructSchemas);
-		Plans = CreateLoadPlans(Batch, GTypes, Bindings, LoadStructIds);
+		Plans = CreateLoadPlans(Batch, GTypes, Customs, GSchemas, LoadStructIds);
 	}
 
 	~FBatchLoader()
@@ -244,32 +317,28 @@ public:
 	}
 
 private:
-	TConstArrayView<FName> SavedNames;
-	FLoadBatch* Plans;
-	TArray<FStructView> Objects;
-	int32 LoadIdx = 0;
+	TConstArrayView<FName>		Ids;
+	FNameBinding				Names;
+	FCustomBindings				Customs;
+	FLoadBatch*					Plans;
+	TArray<FStructView>			Objects;
+	int32						LoadIdx = 0;
 };
 
 
-static void SaveAndLoad(const FStructBindings& Bindings, void (*Save)(FBatchSaver&), void (*Load)(FBatchLoader&))
+static void TestSaveAndLoad(void (*Save)(FBatchSaver&), void (*Load)(FBatchLoader&))
 {
 	TArray64<uint8> Data;
 	{
-		FBatchSaver Batch(Bindings);
+		FBatchSaver Batch;
 		Save(Batch);
 		Data = Batch.Write();
 	}
 
-	FBatchLoader Batch(Bindings, MakeMemoryView(Data));
+	FBatchLoader Batch(MakeMemoryView(Data));
 	Load(Batch);
 }
 
-template<class Runtime = FDefaultRuntime>
-static void TestSaveAndLoad(void (*Save)(FBatchSaver&), void (*Load)(FBatchLoader&))
-{
-	SaveAndLoad(Runtime::GetBindings(), Save, Load);
-}
-
 //////////////////////////////////////////////////////////////////////////
 
 //struct FNestedStructs
@@ -349,6 +418,16 @@ struct FComplexArrays
 PP_REFLECT_STRUCT(PlainProps::UE::Test, FComplexArrays, void, Str, Enums, Misc, Nested);
 static bool operator==(const FComplexArrays& A, const FComplexArrays& B) { return A.Str == B.Str && A.Enums == B.Enums && A.Misc == B.Misc && A.Nested == B.Nested; }
 
+struct FNames
+{
+	FName Name;
+	TArray<FName> Names;
+};
+PP_REFLECT_STRUCT(PlainProps::UE::Test, FNames, void, Name, Names);
+static bool operator==(const FNames& A, const FNames& B) { return A.Name == B.Name && A.Names == B.Names; }
+
+//////////////////////////////////////////////////////////////////////////
+
 struct FUniquePtrs
 {
 	TUniquePtr<bool> Bit;
@@ -387,6 +466,70 @@ TArray<TUniquePtr<T>> MakeTwo(T&& A, T&& B)
 	return Out;
 }
 
+//////////////////////////////////////////////////////////////////////////
+//
+//struct FObject
+//{
+//	virtual ~FObject() {}
+//	char Rtti = '\0';
+//	int Id = 0;
+//};
+//
+//struct FObjectReferenceBinding : public ICustomBinding
+//{
+//	using Type = FObject*;
+//
+//	static constexpr EMemberPresence Occupancy = EMemberPresence::RequireAll;
+//
+//	static TConstArrayView<FMemberId> GetMemberIds()
+//	{
+//		static FMemberId Ids[] = {Ids::IndexMember("Type"), Ids::IndexMember("Id")};
+//		return MakeArrayView(Ids, 2);
+//	}
+//
+//	virtual void SaveStruct(FMemberBuilder& Dst, const void* Src, void*, const FDebugIds&) const override
+//	{
+//		if (const FObject* Object = reinterpret_cast<const FObject*>(Src))
+//		{
+//			Dst.AddLeaf(GetMemberIds()[0], Object->Rtti);
+//			Dst.AddLeaf(GetMemberIds()[1], Object->Id);
+//		}
+//	}
+//
+//	virtual void LoadStruct(void* Dst, FStructView Src, ECustomLoadMethod, const FLoadBatch&) const override
+//	{
+//		FMemberReader Members(Src);
+//		check(Members.PeekName() == GetMemberIds()[0]);
+//		char Rtti = static_cast<char>(Members.GrabLeaf().AsChar8());
+//		check(Members.PeekName() == GetMemberIds()[1]);
+//		int Id = Members.GrabLeaf().AsS32();
+//		check(!Members.HasMore());
+//
+//	}
+//}
+//
+//PP_REFLECT_STRUCT(PlainProps::UE::Test, FObject, void, Rtti, Sibling);
+//
+//struct FObjectX : public FObject
+//{
+//	FObjectX() { Rtti = 'x'; }
+//	FObject* Sibling = nullptr;
+//};
+//PP_REFLECT_STRUCT(PlainProps::UE::Test, FObjectX, FObject, X);
+//
+//struct FObjectY : public FObject
+//{
+//	FObjectY() { Rtti = 'y'; }
+//	FObject* Sibling = nullptr;
+//};
+//PP_REFLECT_STRUCT(PlainProps::UE::Test, FObjectY, FObject, X);
+//
+//struct FOwner
+//{
+//	TArray<TUniquePtr<FObject>> Objects;
+//};
+//PP_REFLECT_STRUCT(PlainProps::UE::Test, FOwner, void, Objects);
+
 //////////////////////////////////////////////////////////////////////////
 
 TEST_CASE_NAMED(FPlainPropsUeCoreTest, "System::Core::Serialization::PlainProps::UE::Core", "[Core][PlainProps][SmokeFilter]")
@@ -407,10 +550,10 @@ TEST_CASE_NAMED(FPlainPropsUeCoreTest, "System::Core::Serialization::PlainProps:
 	
 	SECTION("Enum")
 	{
-		TScopedEnumBinding<EFlat1, EEnumMode::Flat> Flat1;
-		TScopedEnumBinding<EFlat2, EEnumMode::Flat> Flat2;
-		TScopedEnumBinding<EFlag1, EEnumMode::Flag> Flag1;
-		TScopedEnumBinding<EFlag2, EEnumMode::Flag> Flag2;
+		TScopedEnumDeclaration<EFlat1, EEnumMode::Flat> Flat1;
+		TScopedEnumDeclaration<EFlat2, EEnumMode::Flat> Flat2;
+		TScopedEnumDeclaration<EFlag1, EEnumMode::Flag> Flag1;
+		TScopedEnumDeclaration<EFlag2, EEnumMode::Flag> Flag2;
 		TScopedStructBinding<FEnums> Int;
 		TestSaveAndLoad(
 			[](FBatchSaver& Batch)
@@ -453,7 +596,7 @@ TEST_CASE_NAMED(FPlainPropsUeCoreTest, "System::Core::Serialization::PlainProps:
 
 	SECTION("TArrayComplex")
 	{
-		TScopedEnumBinding<EFlat1, EEnumMode::Flat> Flat1;
+		TScopedEnumDeclaration<EFlat1, EEnumMode::Flat> Flat1;
 		TScopedStructBinding<FLeafArrays> LeafArrays;
 		TScopedStructBinding<FComplexArrays> ComplexArrays;
 		TestSaveAndLoad(
@@ -486,7 +629,21 @@ TEST_CASE_NAMED(FPlainPropsUeCoreTest, "System::Core::Serialization::PlainProps:
 			});
 	}
 
-	SECTION("FakeReference")
+	SECTION("FName")
+	{
+		TScopedStructBinding<FNames> Names;
+		TestSaveAndLoad(
+			[](FBatchSaver& Batch)
+			{
+				Batch.Save(FNames{ FName("A"), {FName("Y"), FName("A")} });
+			}, 
+			[](FBatchLoader& Batch)
+			{
+				CHECK(Batch.Load<FNames>() == FNames{ FName("A"), {FName("Y"), FName("A")}});
+			});
+	}
+
+	SECTION("Reference")
 	{}
 
 	SECTION("FString")
@@ -494,9 +651,6 @@ TEST_CASE_NAMED(FPlainPropsUeCoreTest, "System::Core::Serialization::PlainProps:
 		
 	SECTION("TSet")
 	{}
-	
-	SECTION("FName")
-	{}
 
 	SECTION("NestedContainer")
 	{}
