PlainProps new ILeafRangeBinding for non-trivial range transforms w/o resorting to custom struct bindings or upgraders

Also new FString binding that transcodes UTF8 <> UTF16/32 this way

#rnx

diff --git a/Engine/Source/Programs/PlainPropsTests/Private/PlainPropsBind.cpp b/Engine/Source/Programs/PlainPropsTests/Private/PlainPropsBind.cpp
--- a/Engine/Source/Programs/PlainPropsTests/Private/PlainPropsBind.cpp
+++ b/Engine/Source/Programs/PlainPropsTests/Private/PlainPropsBind.cpp
@@ -2,6 +2,7 @@
 
 #include  "PlainPropsBind.h"
 #include  "PlainPropsIndex.h"
+#include  "PlainPropsInternalBuild.h"
 #include  "PlainPropsInternalFormat.h"
 #include  "PlainPropsInternalRead.h"
 
@@ -113,15 +114,33 @@ FSchemaId FMemberVisitor::GrabInnerSchema()
 	return Schema.GetInnerSchemas()[InnerSchemaIdx++];
 }
 
-FRangeBinding::FRangeBinding(const IRangeBinding& Binding, ERangeSizeType SizeType)
+////////////////////////////////////////////////////////////////////////////////////////////////
+
+FRangeBinding::FRangeBinding(const IItemRangeBinding& Binding, ERangeSizeType SizeType)
 : Handle(uint64(&Binding) | uint8(SizeType))
 {
-	check(&Binding == &GetBinding());
+	check(&Binding == &AsItemBinding());
+	check(SizeType == GetSizeType());
+}
+
+FRangeBinding::FRangeBinding(const ILeafRangeBinding& Binding, ERangeSizeType SizeType)
+: Handle(uint64(&Binding) | uint8(SizeType) | LeafMask)
+{
+	check(&Binding == &AsLeafBinding());
 	check(SizeType == GetSizeType());
 }
 
 ////////////////////////////////////////////////////////////////////////////////////////////////
 
+void* FLeafRangeAllocator::Allocate(FUnpackedLeafType Leaf, uint64 Num)
+{
+	check(!Range);
+	Range = FBuiltRange::Create(Num, SizeOf(Leaf.Width));
+	return Range->Data;
+}
+
+////////////////////////////////////////////////////////////////////////////////////////////////
+
 void FCustomBindings::BindStruct(FStructSchemaId Id, ICustomBinding& Binding)
 {
 	checkf(!Find(Id), TEXT("'%s' already bound"), *Debug.Print(Id));
diff --git a/Engine/Source/Programs/PlainPropsTests/Private/PlainPropsBuild.cpp b/Engine/Source/Programs/PlainPropsTests/Private/PlainPropsBuild.cpp
--- a/Engine/Source/Programs/PlainPropsTests/Private/PlainPropsBuild.cpp
+++ b/Engine/Source/Programs/PlainPropsTests/Private/PlainPropsBuild.cpp
@@ -11,6 +11,16 @@
 namespace PlainProps
 {
 
+FBuiltRange* FBuiltRange::Create(uint64 NumItems, SIZE_T ItemSize)
+{
+	check(NumItems > 0);
+	FBuiltRange* Out = new (FMemory::Malloc(sizeof(FBuiltRange) + NumItems * ItemSize)) FBuiltRange;
+	Out->Num = NumItems;
+	return Out;
+}
+
+//////////////////////////////////////////////////////////////////////////
+
 FMemberSchema MakeNestedRangeSchema(ERangeSizeType SizeType, const FMemberSchema& InnerRangeSchema)
 {
 	FMemberSchema Out = { FMemberType(SizeType), InnerRangeSchema.InnerSchema };
@@ -25,10 +35,8 @@ FMemberSchema MakeNestedRangeSchema(ERangeSizeType SizeType, const FMemberSchema
 template<typename BuiltType>
 FBuiltRange* BuildStructuralRangeImpl(/* in-out */ TArrayView64<BuiltType> Values)
 {
-	check(Values.Num() > 0);
 	static_assert(alignof(FBuiltRange) >= alignof(BuiltType));
-	FBuiltRange* Out = new (FMemory::Malloc(sizeof(FBuiltRange) + Values.Num() * sizeof(BuiltType))) FBuiltRange;
-	Out->Num = Values.Num();
+	FBuiltRange* Out = FBuiltRange::Create(Values.Num(), sizeof(FBuiltRange));
 	BuiltType* DataIt = reinterpret_cast<BuiltType*>(Out->Data);
 	for (BuiltType& Value : Values)
 	{
@@ -52,11 +60,8 @@ namespace Private
 
 	FBuiltRange* BuildLeafRange(FUnpackedLeafType Leaf, uint64 Num, FMemoryView Values)
 	{
-		check(Num > 0);
 		check(Values.GetSize() == Num * SizeOf(Leaf.Width));
-
-		FBuiltRange* Out = new (FMemory::Malloc(sizeof(FBuiltRange) + Values.GetSize())) FBuiltRange;
-		Out->Num = Num;
+		FBuiltRange* Out = FBuiltRange::Create(Num, SizeOf(Leaf.Width));
 		FMemory::Memcpy(Out->Data, Values.GetData(), Values.GetSize());
 		return Out;
 	}
diff --git a/Engine/Source/Programs/PlainPropsTests/Private/PlainPropsInternalBuild.h b/Engine/Source/Programs/PlainPropsTests/Private/PlainPropsInternalBuild.h
--- a/Engine/Source/Programs/PlainPropsTests/Private/PlainPropsInternalBuild.h
+++ b/Engine/Source/Programs/PlainPropsTests/Private/PlainPropsInternalBuild.h
@@ -50,13 +50,14 @@ struct FBuiltStruct
 struct FBuiltRange
 {
 	~FBuiltRange() = delete;
-	static uint64 Delete(FBuiltRange* Range, FOptionalSchemaId InnerSchema, TConstArrayView<FMemberType> InnerTypes);
+	[[nodiscard]] static FBuiltRange*					Create(uint64 NumItems, SIZE_T ItemSize);
+	static uint64										Delete(FBuiltRange* Range, FOptionalSchemaId InnerSchema, TConstArrayView<FMemberType> InnerTypes);
 
-	uint64				Num;
-	uint8				Data[0];
+	uint64												Num;
+	uint8												Data[0];
 	
-	TConstArrayView64<const FBuiltRange*> AsRanges() const { return { reinterpret_cast<FBuiltRange const* const*>(Data), IntCastChecked<int64>(Num) }; }
-	TConstArrayView64<TUniquePtr<const FBuiltStruct>> AsStructs() const { return { reinterpret_cast<const TUniquePtr<const FBuiltStruct>*>(Data), IntCastChecked<int64>(Num) }; }
+	TConstArrayView64<const FBuiltRange*>				AsRanges() const { return { reinterpret_cast<FBuiltRange const* const*>(Data), IntCastChecked<int64>(Num) }; }
+	TConstArrayView64<TUniquePtr<const FBuiltStruct>>	AsStructs() const { return { reinterpret_cast<const TUniquePtr<const FBuiltStruct>*>(Data), IntCastChecked<int64>(Num) }; }
 };
 
 //////////////////////////////////////////////////////////////////////////
diff --git a/Engine/Source/Programs/PlainPropsTests/Private/PlainPropsLoad.cpp b/Engine/Source/Programs/PlainPropsTests/Private/PlainPropsLoad.cpp
--- a/Engine/Source/Programs/PlainPropsTests/Private/PlainPropsLoad.cpp
+++ b/Engine/Source/Programs/PlainPropsTests/Private/PlainPropsLoad.cpp
@@ -51,18 +51,6 @@ public:
 	const FSchemaBinding&		AsSchema() const		{ check(IsSchema()); return *AsPtr<FSchemaBinding>(); }
 
 private:
-	// This bit is always zero in user mode addresses and most likely won't be used by current or future
-	// CPU features like ARM's PAC / Top-Byte Ignore or Intel's Linear Address Masking / 5-Level Paging
-#if defined(__x86_64__) || defined(_M_X64)
-	static constexpr uint32 KernelAddressBit = 63;
-#elif defined(__aarch64__) || defined(_M_ARM64)
-	static constexpr uint32 KernelAddressBit = 55;
-#else
-	#error Unsupported architecture, please declare which address bit distinguish user space from kernel space
-#endif
-	// todo handle WASM, copy updated KernelAddressBit from AssetDataTagMap.h
-
-	
 	static constexpr uint64 SparseMask			= uint64(1) << KernelAddressBit;
 	static constexpr uint64 PtrMask				= ~(SparseMask | 0b111);
 	static constexpr uint64 LoMask				= 0b11;
@@ -386,27 +374,39 @@ public:
 
 	static void LoadRange(uint8* Member, FByteReader& ByteIt, FBitCacheReader& BitIt, const FLoadBatch& Batch, const FLoadRangePlan& Range)
 	{
+		uint64 Num = GrabRangeNum(Range.MaxSize, ByteIt, BitIt);
+		FRangeBinding Binding = Range.Bindings[0];
 		FMemberBindType InnerType = Range.InnerTypes[0];
-		const IRangeBinding& Binding = Range.Bindings[0].GetBinding();
-
-		if (uint64 Num = GrabRangeNum(Range.MaxSize, ByteIt, BitIt))
+		
+		if (Binding.IsLeafBinding())
 		{
+			LoadLeafRange(Member, Num, Binding.AsLeafBinding(), ByteIt, UnpackNonBitfield(InnerType.AsLeaf()));
+		}
+		else if (Num)
+		{
+			const IItemRangeBinding& ItemBinding = Binding.AsItemBinding();
 			switch (InnerType.GetKind())
 			{
-				case EMemberKind::Leaf:		LoadRangeValues(Member, Num, Binding, ByteIt, Batch, InnerType.AsLeaf()); break;
-				case EMemberKind::Range:	LoadRangeValues(Member, Num, Binding, ByteIt, Batch, Range.Tail()); break;
-				case EMemberKind::Struct:	LoadRangeValues(Member, Num, Binding, ByteIt, Batch, Range.InnermostStruct.Get()); break;
+				case EMemberKind::Leaf:		LoadRangeValues(Member, Num, ItemBinding, ByteIt, Batch, UnpackNonBitfield(InnerType.AsLeaf())); break;
+				case EMemberKind::Range:	LoadRangeValues(Member, Num, ItemBinding, ByteIt, Batch, Range.Tail()); break;
+				case EMemberKind::Struct:	LoadRangeValues(Member, Num, ItemBinding, ByteIt, Batch, Range.InnermostStruct.Get()); break;
 			}
 		}
 		else
 		{
 			FLoadRangeContext NoItemsCtx{.Request = {Member, 0}};
-			(Binding.MakeItems)(NoItemsCtx);
+			(Binding.AsItemBinding().MakeItems)(NoItemsCtx);
 		}
+	}	
+	
+	static void LoadLeafRange(uint8* Member, uint64 Num, const ILeafRangeBinding& Binding, FByteReader& ByteIt, FUnpackedLeafType Leaf)
+	{
+		FMemoryView Values = Num ? GrabRangeValues(ByteIt, Num, Leaf) : FMemoryView();
+		Binding.LoadLeaves(Member, FLeafRangeLoadView(Values.GetData(), Num, Leaf));
 	}
 
 	template<class SchemaType>
-	static void LoadRangeValues(uint8* Member, uint64 Num, const IRangeBinding& Binding, FByteReader& ByteIt, const FLoadBatch& Batch, SchemaType&& Schema)
+	static void LoadRangeValues(uint8* Member, uint64 Num, const IItemRangeBinding& Binding, FByteReader& ByteIt, const FLoadBatch& Batch, SchemaType&& Schema)
 	{
 		FByteReader ValueIt(GrabRangeValues(ByteIt, Num, Schema));
 		FBitCacheReader BitIt; // Only used by ranges of ERangeSizeType::Uni ranges
@@ -426,10 +426,10 @@ public:
 		}
 	}
 
-	static FMemoryView GrabRangeValues(FByteReader& ByteIt, uint64 Num, FLeafBindType Schema)
+	static FMemoryView GrabRangeValues(FByteReader& ByteIt, uint64 Num, FUnpackedLeafType Leaf)
 	{
 		check(Num > 0);
-		return ByteIt.GrabSlice(GetLeafRangeSize(Num, ToLeafType(Schema)));
+		return ByteIt.GrabSlice(GetLeafRangeSize(Num, Leaf));
 	}
 
 	template<class SchemaType>
@@ -438,27 +438,21 @@ public:
 		return ByteIt.GrabSkippableSlice();
 	}
 		
-	static void CopyRangeValues(const FConstructedItems& Items, FByteReader& ByteIt, FBitCacheReader&, const FLoadBatch& Batch, FLeafBindType Leaf)
+	static void CopyRangeValues(const FConstructedItems& Items, FByteReader& ByteIt, FBitCacheReader&, const FLoadBatch& Batch, FUnpackedLeafType Leaf)
 	{
-		switch (Leaf.Bind.Type)
+		check(Items.Size == SizeOf(Leaf.Width));
+		if (Leaf.Type != ELeafType::Bool)
 		{
-		case ELeafBindType::Bool:
+			FMemory::Memcpy(Items.Data, ByteIt.GrabBytes(Items.NumBytes()), Items.NumBytes());
+		}
+		else
 		{
-			check(Items.Size == sizeof(bool));
 			FBoolRangeView Bits(ByteIt.GrabBytes(Align(Items.Num, 8)/8), Items.Num);
 			uint8* It = Items.Data;
 			for (bool bBit : Bits)
 			{
 				reinterpret_cast<bool&>(*It++) = bBit;
 			}
-
-			break;
-		}
-		default:
-			checkf(Leaf.Bind.Type != ELeafBindType::BitfieldBool, TEXT("Loading to bitfield array is unsupported"));
-			check(Items.Size == SizeOf(Leaf.Arithmetic.Width));
-			FMemory::Memcpy(Items.Data, ByteIt.GrabBytes(Items.NumBytes()), Items.NumBytes());
-			break;
 		}
 	}
 	
diff --git a/Engine/Source/Programs/PlainPropsTests/Private/PlainPropsSave.cpp b/Engine/Source/Programs/PlainPropsTests/Private/PlainPropsSave.cpp
--- a/Engine/Source/Programs/PlainPropsTests/Private/PlainPropsSave.cpp
+++ b/Engine/Source/Programs/PlainPropsTests/Private/PlainPropsSave.cpp
@@ -48,150 +48,19 @@ static uint64 SaveLeaf(const uint8* Member, FUnpackedLeafBindType Leaf)
 	return uint64(0);
 }
 
-//FBuiltRange* SaveLeafRange(const void* Range, const IRangeBinding& Binding, FUnpackedLeafType Leaf)
-//{
-//	FSaveRangeContext Ctx = { { Range } };
-//
-//	Binding.ReadItems(Ctx);
-//
-//	if (const uint64 NumTotal = Ctx.Items.NumTotal)
-//	{
-//		FBuiltRange* Out = new (FMemory::Malloc(sizeof(FBuiltRange) + NumTotal * SizeOf(Leaf.Width))) FBuiltRange;
-//		Out->Num = NumTotal;
-//
-//		for (uint64 NumRead = 0; true; NumRead += Ctx.Items.Slice.Num)
-//		{
-//			check(Ctx.Items.Slice.Num > 0);
-//
-//			FMemory::Memcpy(Out->Data + NumRead, Ctx.Items.Slice.Data,  Ctx.Items.Slice.Num * SizeOf(Leaf.Width));
-//
-//			if (NumRead >= NumTotal)
-//			{
-//				check(NumRead == NumTotal);	
-//				break;
-//			}
-//
-//			Binding.ReadItems(Ctx);	
-//		}
-//
-//		return Out;
-//	}
-//	
-//	return nullptr;
-//}
-//
-//FBuiltRange* SaveStructRange(const void* Range, const IRangeBinding& Binding, FStructType Type, FStructSchemaId Id, const FSaveContext& OuterCtx)
-//{
-//	FSaveRangeContext Ctx = { { Range } };
-//
-//	Binding.ReadItems(Ctx);
-//
-//	if (const uint64 NumTotal = Ctx.Items.NumTotal)
-//	{
-//		TArray64<TUniquePtr<FBuiltStruct>> Structs;
-//		Structs.Reserve(NumTotal);
-//
-//		for (uint64 NumRead = 0; true; NumRead += Ctx.Items.Slice.Num)
-//		{
-//			check(Ctx.Items.Slice.Num > 0);
-//
-//			for (uint64 Idx = 0, Num = Ctx.Items.Slice.Num; Idx < Num; ++Idx)
-//			{
-//				Structs.Emplace(SaveStruct(Ctx.Items.Slice.At(Idx, Ctx.Items.Stride), Id, OuterCtx));
-//			}
-//
-//			if (NumRead >= NumTotal)
-//			{
-//				check(NumRead == NumTotal);	
-//				break;
-//			}
-//
-//			Binding.ReadItems(Ctx);	
-//		}
-//
-//		return Private::BuildStructRange(/* ownership xfer */ Structs);
-//	}
-//	
-//	return nullptr;
-//}
-
-
-
-//
-//static FBuiltRange* SaveNestedRange(const void* Range, const IRangeBinding& Binding, FRangeType Type, FRangeMemberBinding InnerBinding, const FSaveContext& OuterCtx)
-//{
-//	FSaveRangeContext Ctx = { { Range } };
-//
-//	Binding.ReadItems(Ctx);
-//
-//	if (const uint64 NumTotal = Ctx.Items.NumTotal)
-//	{
-//		TArray64<FBuiltRange*> Ranges;
-//		Ranges.Reserve(NumTotal);
-//
-//		for (uint64 NumRead = 0; true; NumRead += Ctx.Items.Slice.Num)
-//		{
-//			check(Ctx.Items.Slice.Num > 0);
-//
-//			for (uint64 Idx = 0, Num = Ctx.Items.Slice.Num; Idx < Num; ++Idx)
-//			{
-//				Ranges.Emplace(SaveRange(Ctx.Items.Slice.At(Idx, Ctx.Items.Stride), InnerBinding, OuterCtx));
-//			}
-//
-//			if (NumRead >= NumTotal)
-//			{
-//				check(NumRead == NumTotal);	
-//				break;
-//			}
-//
-//			Binding.ReadItems(Ctx);	
-//		}
-//
-//		return Private::BuildNestedRange(/* ownership xfer */ Ranges);
-//	}
-//	
-//	return nullptr;
-//}
-
-
-//struct FStructRangeSaver
-//{
-//	FStructSchemaId Id;
-//	TArray64<TUniquePtr<FBuiltStruct>> Structs;
-//
-//	void Init(uint64 Num)
-//	{ 
-//		Ranges.Reserve(NumTotal);
-//	}
-//
-//	void Append(FExistingItemSlice Slice, uint32 Stride, FSaveContext& OuterCtx)
-//	{
-//		for (uint64 Idx = 0; Idx < Slice.Num; ++Idx)
-//		{
-//			Ranges.Emplace(SaveRange(Slice.At(Idx, Stride), Id, OuterCtx));
-//		}
-//	}
-//
-//	FBuiltRange* Build()
-//	{
-//		Private::BuildStructRange(/* ownership xfer */ Structs);
-//	}
-//};
-
 struct FLeafRangeSaver
 {
 	FBuiltRange* Out;
 	uint8* OutIt;
 
 	FLeafRangeSaver(uint64 Num, SIZE_T LeafSize)
-	: Out(new (FMemory::Malloc(sizeof(FBuiltRange) + Num * LeafSize)) FBuiltRange)
+	: Out(FBuiltRange::Create(Num, LeafSize))
 	, OutIt(Out->Data)
-	{ 
-		Out->Num = Num;
-	}
+	{}
 
-	void Append(FExistingItemSlice Slice, uint32, SIZE_T LeafSize, const FSaveContext&)
+	void Append(FExistingItemSlice Slice, uint32 Stride, SIZE_T LeafSize, const FSaveContext&)
 	{
+		check(Stride == LeafSize);
 		FMemory::Memcpy(OutIt, Slice.Data, Slice.Num * LeafSize);
 		OutIt += Slice.Num * LeafSize;
 	}
@@ -202,6 +71,8 @@ struct FLeafRangeSaver
 	}
 };
 
+//////////////////////////////////////////////////////////////////////////
+
 template<typename BuiltItemType, typename ItemSchemaType>
 struct TStructuralRangeSaver
 {
@@ -229,8 +100,10 @@ struct TStructuralRangeSaver
 using FNestedRangeSaver = TStructuralRangeSaver<FBuiltRange*, FRangeMemberBinding>;
 using FStructRangeSaver = TStructuralRangeSaver<TUniquePtr<FBuiltStruct>, FStructSchemaId>;
 
+//////////////////////////////////////////////////////////////////////////
+
 template<class SaverType, typename InnerContextType>
-[[nodiscard]] FBuiltRange* SaveRange(const void* Range, const IRangeBinding& Binding, const FSaveContext& OuterCtx, InnerContextType InnerCtx)
+[[nodiscard]] FBuiltRange* SaveRange(const void* Range, const IItemRangeBinding& Binding, const FSaveContext& OuterCtx, InnerContextType InnerCtx)
 {
 	FSaveRangeContext ReadCtx = { { Range } };
 	Binding.ReadItems(ReadCtx);
@@ -243,10 +116,10 @@ template<class SaverType, typename InnerContextType>
 			check(ReadCtx.Items.Slice.Num > 0);
 			Saver.Append(ReadCtx.Items.Slice, ReadCtx.Items.Stride, InnerCtx, OuterCtx);
 		
-			ReadCtx.Request.Index += ReadCtx.Items.Slice.Num;
-			if (ReadCtx.Request.Index >= NumTotal)
+			ReadCtx.Request.NumRead += ReadCtx.Items.Slice.Num;
+			if (ReadCtx.Request.NumRead >= NumTotal)
 			{
-				check(ReadCtx.Request.Index == NumTotal);	
+				check(ReadCtx.Request.NumRead == NumTotal);	
 				return Saver.Finish();
 			}
 
@@ -270,21 +143,37 @@ template<class SaverType, typename InnerContextType>
 	return Leaf.Arithmetic.Width;
  }
 
+[[nodiscard]] static FBuiltRange* SaveLeafRange(const uint8* Range, const ILeafRangeBinding& Binding, FUnpackedLeafType Leaf)
+{
+	FLeafRangeAllocator Allocator(Leaf);
+	Binding.SaveLeaves(Range, Allocator);
+	return Allocator.GetAllocatedRange();
+}
+
 [[nodiscard]] static FBuiltRange* SaveRange(const uint8* Range, FRangeMemberBinding Member, const FSaveContext& Ctx)
 {
-	const IRangeBinding& Binding = Member.RangeBindings[0].GetBinding();
+	FRangeBinding Binding = Member.RangeBindings[0];
 	FMemberBindType InnerType = Member.InnerTypes[0];
+
+	if (Binding.IsLeafBinding())
+	{
+		return SaveLeafRange(Range, Binding.AsLeafBinding(), UnpackNonBitfield(InnerType.AsLeaf()));
+	}
+
+	const IItemRangeBinding& ItemBinding = Binding.AsItemBinding();
 	switch (InnerType.GetKind())
 	{
-	case EMemberKind::Leaf:		return SaveRange<FLeafRangeSaver>(  Range, Binding, Ctx, SizeOf(GetArithmeticWidth(InnerType.AsLeaf())));
-	case EMemberKind::Range:	return SaveRange<FNestedRangeSaver>(Range, Binding, Ctx, GetInnerRange(Member));
-	case EMemberKind::Struct:	return SaveRange<FStructRangeSaver>(Range, Binding, Ctx, static_cast<FStructSchemaId>(Member.InnermostSchema.Get()));
+	case EMemberKind::Leaf:		return SaveRange<FLeafRangeSaver>(  Range, ItemBinding, Ctx, SizeOf(GetArithmeticWidth(InnerType.AsLeaf())));
+	case EMemberKind::Range:	return SaveRange<FNestedRangeSaver>(Range, ItemBinding, Ctx, GetInnerRange(Member));
+	case EMemberKind::Struct:	return SaveRange<FStructRangeSaver>(Range, ItemBinding, Ctx, static_cast<FStructSchemaId>(Member.InnermostSchema.Get()));
 	}
 
 	check(false);
 	return nullptr;
 }
 
+//////////////////////////////////////////////////////////////////////////
+
 [[nodiscard]] static FBuiltRange* SaveRangeItem(const uint8* Range, FRangeMemberBinding Member, const FSaveContext& Ctx)
 { 
 	return SaveRange(Range, Member, Ctx);
diff --git a/Engine/Source/Programs/PlainPropsTests/Public/PlainPropsBind.h b/Engine/Source/Programs/PlainPropsTests/Public/PlainPropsBind.h
--- a/Engine/Source/Programs/PlainPropsTests/Public/PlainPropsBind.h
+++ b/Engine/Source/Programs/PlainPropsTests/Public/PlainPropsBind.h
@@ -5,23 +5,24 @@
 #include "Containers/ArrayView.h"
 #include "Memory/MemoryFwd.h"
 #include "Memory/MemoryView.h"
-#include  "PlainPropsTypes.h"
 #include  "PlainPropsCtti.h"
 #include  "PlainPropsDeclare.h"
+#include  "PlainPropsRead.h"
+#include  "PlainPropsTypes.h"
 
 namespace PlainProps 
 {
 
+struct FBuiltRange;
 struct FBuiltStruct;
 class FIdIndexerBase;
 struct FLoadBatch;
 class FMemberBuilder;
 struct FSchemaBatch;
 class FStructBinding;
-struct FStructView;
 class FRangeBinding;
 struct FTypedRange;
-class IRangeBinding;
+class IItemRangeBinding;
 template<class T> class TIdIndexer;
 
 ////////////////////////////////////////////////////////////////////////////////////////////////
@@ -163,6 +164,13 @@ struct FUnpackedLeafBindType
 	}
 };
 
+// @pre Type != ELeafBindType::BitfieldBool
+inline FUnpackedLeafType UnpackNonBitfield(FLeafBindType Packed)
+{
+	FUnpackedLeafBindType Unpacked(Packed);
+	return { ToLeafType(Unpacked.Type), Unpacked.Width };
+}
+
 struct FLeafMemberBinding
 {
 	FUnpackedLeafBindType	Leaf;
@@ -265,7 +273,7 @@ class FConstructionRequest
 	const uint64 Num = 0;
 	uint64 Index = 0;
 
-	friend class FRangeLoader;
+	friend FRangeLoader;
 	FConstructionRequest(void* InRange, uint64 InNum) : Range(InRange), Num(InNum) {}
 	
 public:
@@ -309,7 +317,7 @@ public:
 	}
 
 private:
-	friend class FRangeLoader;
+	friend FRangeLoader;
 	uint8*	Data = nullptr;
 	uint64	Num = 0;			
 	uint32	Size = 0;
@@ -322,8 +330,8 @@ private:
 struct FLoadRangeContext
 {
 	FConstructionRequest	Request;		// Request to construct items to be loaded
-	FConstructedItems		Items;			// Response from IRangeBinding
-	uint64					Scratch[64];	// Scratch memory for IRangeBinding
+	FConstructedItems		Items;			// Response from IItemRangeBinding
+	uint64					Scratch[64];	// Scratch memory for IItemRangeBinding
 };
 
 // todo: switch to class
@@ -332,8 +340,10 @@ struct FGetItemsRequest
 	template<typename T>
 	const T& GetRange() const { return *reinterpret_cast<const T*>(Range); }
 
+	bool IsFirstCall() const { return NumRead == 0;}
+
 	const void* Range = nullptr;
-	uint64 Index = 0;
+	uint64 NumRead = 0;
 };
 
 struct FExistingItemSlice
@@ -372,39 +382,118 @@ struct FExistingItems
 	{
 		SetAll(FExistingItemSlice{Items, NumItems}, sizeof(ItemType));
 	}
-
-	//bool HasAll() const
-	//{
-	//	return NumTotal == Slice.Num;
-	//}
 };
 
+
+// TODO: Consider API changes so that
+//  - Leaf ranges can allocate data and fill it in directly
+//  - Leaf ranges can allocate more capacity than is used
+//  - Maybe separate Finalize function
+//  - Continue hiding intermediate details (e.g. FBuiltRange) and allocator
 struct FSaveRangeContext
 {
-	FGetItemsRequest		Request;	// Request to get items to be saved
-	FExistingItems			Items;		// Response from IRangeBinding
-	uint64					Scratch[8]; // Scratch memory for IRangeBinding
+	FGetItemsRequest		Request;		// Request to get items to be saved
+	FExistingItems			Items;			// Response from IRangeBinding
+	uint64					Scratch[8]; 	// Scratch memory for IRangeBinding
 };
 
-class IRangeBinding
+class alignas(16) IItemRangeBinding
 {
 public:
-	virtual void MakeItems(FLoadRangeContext& Ctx) const = 0;
 	virtual void ReadItems(FSaveRangeContext& Ctx) const = 0;
+	virtual void MakeItems(FLoadRangeContext& Ctx) const = 0;
 };
 
 ////////////////////////////////////////////////////////////////////////////////////////////////
 
+// Possible save opt: Use paged linear allocator that only allocates on page exhaustion
+class FLeafRangeAllocator
+{
+	const FUnpackedLeafType		Expected;
+	FBuiltRange*				Range = nullptr;
+
+	void* Allocate(FUnpackedLeafType Type, uint64 Num);
+
+public:
+	FLeafRangeAllocator(FUnpackedLeafType InExpected) : Expected(InExpected) {}
+
+	template<typename LeafType, typename SizeType>
+	LeafType* AllocateRange(SizeType Num)
+	{
+		check(ReflectLeaf<LeafType> == Expected);
+		return Num ? static_cast<LeafType*>(Allocate(ReflectLeaf<LeafType>, IntCastChecked<uint64>(Num))) : nullptr;
+	}
+
+	FBuiltRange* GetAllocatedRange() { return Range; }
+};
+
+class FLeafRangeLoadView
+{
+	const void*				Data;
+	uint64					Num;
+	FUnpackedLeafType		Leaf;
+
+public:
+	FLeafRangeLoadView(const void* InData, uint64 InNum, FUnpackedLeafType InLeaf)
+	: Data(InData)
+	, Num(InNum)
+	, Leaf(InLeaf)
+	{}
+	
+	// The returned ranges hide the internal representations so we can change format in the future, 
+	// e.g. store zeroes or 1.0f in some compact fashion or even var int encodings
+	template<typename LeafType>
+	auto As() const
+	{
+		check(Leaf == ReflectLeaf<LeafType>);
+		if constexpr (std::is_same_v<bool, LeafType>)
+		{
+			return FBoolRangeView(static_cast<const uint8*>(Data), Num);
+		}
+		else
+		{
+			return TRangeView<LeafType>(static_cast<const LeafType*>(Data), Num);
+		}
+	}
+};
+
+// Specialized binding for transcoding leaf ranges
+class alignas(16) ILeafRangeBinding
+{
+public:
+	virtual void SaveLeaves(const void* Range, FLeafRangeAllocator& Out) const = 0;
+	virtual void LoadLeaves(void* Range, FLeafRangeLoadView Leaves) const = 0;
+};
+
+////////////////////////////////////////////////////////////////////////////////////////////////
+
+// Todo: Switch to FPlatformMemory::KernelAddressBit once that is submitted
+//
+// This bit is always zero in user mode addresses and most likely won't be used by current or future
+// CPU features like ARM's PAC / Top-Byte Ignore or Intel's Linear Address Masking / 5-Level Paging
+#if defined(__x86_64__) || defined(_M_X64)
+	static constexpr uint32 KernelAddressBit = 63;
+#elif defined(__aarch64__) || defined(_M_ARM64)
+	static constexpr uint32 KernelAddressBit = 55;
+#else
+	#error Unsupported architecture, please declare which address bit distinguish user space from kernel space
+#endif
+
 class FRangeBinding
 {
-	static constexpr uint64		SizeMask = 0b111;
+	static constexpr uint64		SizeMask = 0b1111;
+	static constexpr uint64		LeafMask = uint64(1) << KernelAddressBit;
+	static constexpr uint64		BindMask = ~(SizeMask | LeafMask);
 	uint64						Handle;
 
 public:
-	FRangeBinding(const IRangeBinding& Binding, ERangeSizeType SizeType);
-
-	const IRangeBinding&		GetBinding() const	{ return *reinterpret_cast<IRangeBinding*>(Handle & ~SizeMask); }
-	ERangeSizeType				GetSizeType() const { return static_cast<ERangeSizeType>(Handle & SizeMask); }
+	FRangeBinding(const IItemRangeBinding& Binding, ERangeSizeType SizeType);
+	FRangeBinding(const ILeafRangeBinding& Binding, ERangeSizeType SizeType);
+	
+	bool						IsLeafBinding() const		{ return !!(LeafMask & Handle); }
+	const IItemRangeBinding&	AsItemBinding() const		{ check(!IsLeafBinding()); return *reinterpret_cast<IItemRangeBinding*>(Handle & BindMask); }
+	const ILeafRangeBinding&	AsLeafBinding() const		{ check( IsLeafBinding()); return *reinterpret_cast<ILeafRangeBinding*>(Handle & BindMask); }
+	ERangeSizeType				GetSizeType() const			{ return static_cast<ERangeSizeType>(Handle & SizeMask); }
 };
 
 template<typename T>
@@ -425,8 +514,6 @@ struct FMemberBinding
 	TConstArrayView<FRangeBinding>	RangeBindings;		// Non-empty -> Range
 };
 
-////////////////////////////////////////////////////////////////////////////////////////////////
-
 class FSchemaBindings
 {
 public:
@@ -799,5 +886,4 @@ struct FIdTranslator : FIdTranslatorBase
 FSchemaBatch*			CreateTranslatedSchemas(const FSchemaBatch& Schemas, FIdBinding NewIds);
 void					DestroyTranslatedSchemas(const FSchemaBatch* Schemas);
 
-} // namespace PlainProps
-
+} // namespace PlainProps
\ No newline at end of file
diff --git a/Engine/Source/Programs/PlainPropsTests/Public/PlainPropsRead.h b/Engine/Source/Programs/PlainPropsTests/Public/PlainPropsRead.h
--- a/Engine/Source/Programs/PlainPropsTests/Public/PlainPropsRead.h
+++ b/Engine/Source/Programs/PlainPropsTests/Public/PlainPropsRead.h
@@ -250,8 +250,8 @@ public:
 	bool						IsNestedRange() const	{ return Schema.ItemType.GetKind() == EMemberKind::Range; }
 
 	PP_API FLeafRangeView		AsLeaves() const;		// @pre IsLeafRange()
-	PP_API FStructRangeView	AsStructs() const;		// @pre IsStructRange()
-	PP_API FNestedRangeView	AsRanges() const;		// @pre IsNestedRange()
+	PP_API FStructRangeView		AsStructs() const;		// @pre IsStructRange()
+	PP_API FNestedRangeView		AsRanges() const;		// @pre IsNestedRange()
 
 private:
 	friend FMemberReader;
@@ -401,6 +401,7 @@ public:
 	template<typename LeafType>
 	TRangeView<LeafType> As() const
 	{
+		static_assert(!std::is_same_v<LeafType, bool>);
 		check(FUnpackedLeafType(Type, Width) == ReflectLeaf<LeafType>);
 		return TRangeView<LeafType>(reinterpret_cast<const LeafType*>(Values), NumItems);
 	}
diff --git a/Engine/Source/Programs/PlainPropsTests/Public/PlainPropsTypes.h b/Engine/Source/Programs/PlainPropsTests/Public/PlainPropsTypes.h
--- a/Engine/Source/Programs/PlainPropsTests/Public/PlainPropsTypes.h
+++ b/Engine/Source/Programs/PlainPropsTests/Public/PlainPropsTypes.h
@@ -318,7 +318,7 @@ struct FUnpackedLeafType
 	constexpr FUnpackedLeafType(ELeafType InType, ELeafWidth InWidth) : Type(InType), Width(InWidth) {}
 	constexpr FUnpackedLeafType(FLeafType In) : Type(In.Type), Width(In.Width) {}
 
-	constexpr bool operator==(FUnpackedLeafType O) { return Type == O.Type && Width == O.Width; }
+	constexpr bool operator==(FUnpackedLeafType O) const { return Type == O.Type && Width == O.Width; }
 	FMemberType Pack() const { return FMemberType(Type, Width); }
 };
 
diff --git a/Engine/Source/Programs/PlainPropsTests/Public/PlainPropsUeCoreBindings.h b/Engine/Source/Programs/PlainPropsTests/Public/PlainPropsUeCoreBindings.h
--- a/Engine/Source/Programs/PlainPropsTests/Public/PlainPropsUeCoreBindings.h
+++ b/Engine/Source/Programs/PlainPropsTests/Public/PlainPropsUeCoreBindings.h
@@ -99,12 +99,13 @@ namespace PlainProps::UE
 
 #include "Containers/Array.h"
 #include "Containers/Set.h"
+#include "Templates/UniquePtr.h"
 
 namespace PlainProps::UE
 {
 
 template <typename T>
-struct TArrayBinding : public IRangeBinding
+struct TArrayBinding : public IItemRangeBinding
 {
 	using SizeType = int32;
 	using ItemType = T;
@@ -134,8 +135,63 @@ struct TArrayBinding : public IRangeBinding
 
 //////////////////////////////////////////////////////////////////////////
 
+struct FStringBinding : public ILeafRangeBinding
+{
+	using SizeType = int32;
+	using ItemType = char8_t;
+
+	virtual void SaveLeaves(const void* Range, FLeafRangeAllocator& Out) const override
+	{
+		const TArray<TCHAR>& Src = static_cast<const FString*>(Range)->GetCharArray();
+		int32 SrcLen = Src.Num() - 1;
+		if (SrcLen <= 0)
+		{
+		}
+		else if constexpr (sizeof(TCHAR) == sizeof(char8_t))
+		{
+			char8_t* Utf8 = Out.AllocateRange<char8_t>(SrcLen);
+			FMemory::Memcpy(Utf8, Src.GetData(), SrcLen);
+		}
+		else
+		{
+			int32 Utf8Len = FPlatformString::ConvertedLength<UTF8CHAR>(Src.GetData(), SrcLen);
+			char8_t* Utf8 = Out.AllocateRange<char8_t>(Utf8Len);
+			UTF8CHAR* Utf8End = FPlatformString::Convert(reinterpret_cast<UTF8CHAR*>(Utf8), Utf8Len, Src.GetData(), SrcLen);	
+			check((char8_t*)Utf8End - Utf8 == Utf8Len);
+		}
+	}
+
+	virtual void LoadLeaves(void* Range, FLeafRangeLoadView Items) const override
+	{
+		TArray<TCHAR>& Dst = static_cast<FString*>(Range)->GetCharArray();
+		TRangeView<char8_t> Utf8 = Items.As<char8_t>();
+		const UTF8CHAR* Src = reinterpret_cast<const UTF8CHAR*>(Utf8.begin());
+		int32 SrcLen = static_cast<int32>(Utf8.Num());
+		if (SrcLen == 0)
+		{
+			Dst.Reset();
+		}
+		else if constexpr (sizeof(TCHAR) == sizeof(char8_t))
+		{
+			Dst.SetNum(SrcLen + 1);
+			FMemory::Memcpy(Dst.GetData(), Src, SrcLen);
+			Dst[SrcLen] = '\0';	
+		}
+		else
+		{
+			int32 DstLen = FPlatformString::ConvertedLength<TCHAR>(Src, SrcLen);
+			Dst.SetNum(DstLen + 1);
+			TCHAR* DstEnd = FPlatformString::Convert(Dst.GetData(), DstLen, Src, SrcLen);
+			check(DstEnd - Dst.GetData() == DstLen);
+			*DstEnd = '\0';
+		}
+	}
+};
+
+//////////////////////////////////////////////////////////////////////////
+
 template <typename T>
-struct TUniquePtrBinding : public IRangeBinding
+struct TUniquePtrBinding : public IItemRangeBinding
 {
 	using SizeType = bool;
 	using ItemType = T;
@@ -177,7 +233,7 @@ struct TUniquePtrBinding : public IRangeBinding
 //////////////////////////////////////////////////////////////////////////
 
 template <typename T>
-struct TSetBinding : public IRangeBinding
+struct TSetBinding : public IItemRangeBinding
 {
 	using SizeType = int32;
 	using ItemType = T;
@@ -279,9 +335,7 @@ struct TSetBinding : public IRangeBinding
 	}
 };
 
-
 //////////////////////////////////////////////////////////////////////////
-
 struct FSetOps
 {
 	union
@@ -312,7 +366,7 @@ struct TSetDeltaBinding : public ICustomBinding
 	
 	static TConstArrayView<FMemberId> GetMemberIds() { return FSetOps::Get<Ids>().All; }
 
-	virtual void SaveStruct(FMemberBuilder& Dst, const void* Src, void* UserData, const FDebugIds& Debug) const override;
+	virtual void SaveStruct(FMemberBuilder& Dst, const void* Src, void* UserData, const FDebugIds& Debug) override;
 
 	virtual void LoadStruct(void* Dst, FStructView Src, ECustomLoadMethod Method, const FLoadBatch& Batch) const override
 	{
@@ -394,7 +448,7 @@ struct TSetDeltaBinding : public ICustomBinding
 		else if constexpr (std::is_default_constructible_v<T>())
 		{
 			using Binding = RangeBind<T>;
-			const IRangeBinding* Bindings[] = {};// ... generate somehow ... };
+			const IItemRangeBinding* Bindings[] = {};// ... generate somehow ... };
 			T Tmp;
 			for (FRangeView Item : Items.AsRanges())
 			{
@@ -483,6 +537,12 @@ namespace PlainProps
 		using Type = UE::TArrayBinding<T>;
 	};
 
+	template<>
+	struct TRangeBind<FString>
+	{
+		using Type = UE::FStringBinding;
+	};
+
 	template<typename T>
 	struct TRangeBind<TUniquePtr<T>>
 	{
diff --git a/Engine/Source/Programs/PlainPropsTests/Tests/PlainPropsUeCoreTests.cpp b/Engine/Source/Programs/PlainPropsTests/Tests/PlainPropsUeCoreTests.cpp
--- a/Engine/Source/Programs/PlainPropsTests/Tests/PlainPropsUeCoreTests.cpp
+++ b/Engine/Source/Programs/PlainPropsTests/Tests/PlainPropsUeCoreTests.cpp
@@ -130,7 +130,12 @@ struct FNameDeclaration
 	}
 };
 
-struct FNameBinding : public ICustomBinding
+struct FTestCustomBinding : public ICustomBinding
+{
+	virtual FStructSchemaId GetId() const = 0;
+};
+
+struct FNameBinding : public FTestCustomBinding
 {
 	virtual void SaveStruct(FMemberBuilder& Dst, const void* Src, void*, const FDebugIds& Debug) override
 	{
@@ -144,6 +149,8 @@ struct FNameBinding : public ICustomBinding
 		*static_cast<FName*>(Dst) = Names.Get(Idx);
 	}
 
+	virtual FStructSchemaId GetId() const override { return Declaration.Id; }
+
 	FNameDeclaration	Declaration;
 	TSet<FName>			Names;
 };
@@ -326,7 +333,7 @@ private:
 };
 
 
-static void TestSaveAndLoad(void (*Save)(FBatchSaver&), void (*Load)(FBatchLoader&))
+static void Run(void (*Save)(FBatchSaver&), void (*Load)(FBatchLoader&))
 {
 	TArray64<uint8> Data;
 	{
@@ -426,6 +433,13 @@ struct FNames
 PP_REFLECT_STRUCT(PlainProps::UE::Test, FNames, void, Name, Names);
 static bool operator==(const FNames& A, const FNames& B) { return A.Name == B.Name && A.Names == B.Names; }
 
+struct FStr
+{
+	FString S;
+};
+PP_REFLECT_STRUCT(PlainProps::UE::Test, FStr, void, S);
+
+
 //////////////////////////////////////////////////////////////////////////
 
 struct FUniquePtrs
@@ -537,8 +551,7 @@ TEST_CASE_NAMED(FPlainPropsUeCoreTest, "System::Core::Serialization::PlainProps:
 	SECTION("Basic")
 	{
 		TScopedStructBinding<FInt> Int;
-		TestSaveAndLoad(
-			[](FBatchSaver& Batch)
+		Run([](FBatchSaver& Batch)
 			{
 				Batch.Save(FInt{1234});
 			}, 
@@ -555,8 +568,7 @@ TEST_CASE_NAMED(FPlainPropsUeCoreTest, "System::Core::Serialization::PlainProps:
 		TScopedEnumDeclaration<EFlag1, EEnumMode::Flag> Flag1;
 		TScopedEnumDeclaration<EFlag2, EEnumMode::Flag> Flag2;
 		TScopedStructBinding<FEnums> Int;
-		TestSaveAndLoad(
-			[](FBatchSaver& Batch)
+		Run([](FBatchSaver& Batch)
 			{
 				Batch.Save(FEnums{EFlat1::A, EFlat2::A, EFlag1::A, EFlag2::A});
 				Batch.Save(FEnums{EFlat1::A, EFlat2::A, EFlag1::B, EFlag2::B});
@@ -574,11 +586,10 @@ TEST_CASE_NAMED(FPlainPropsUeCoreTest, "System::Core::Serialization::PlainProps:
 			});
 	}
 
-	SECTION("TArrayBasic")
+	SECTION("TArray")
 	{
 		TScopedStructBinding<FLeafArrays> LeafArrays;
-		TestSaveAndLoad(
-			[](FBatchSaver& Batch)
+		Run([](FBatchSaver& Batch)
 			{
 				Batch.Save(FLeafArrays{{}, {}});
 				Batch.Save(FLeafArrays{{false}, {1, 2}});
@@ -594,13 +605,12 @@ TEST_CASE_NAMED(FPlainPropsUeCoreTest, "System::Core::Serialization::PlainProps:
 			});
 	}
 
-	SECTION("TArrayComplex")
+	SECTION("Nesting")
 	{
 		TScopedEnumDeclaration<EFlat1, EEnumMode::Flat> Flat1;
 		TScopedStructBinding<FLeafArrays> LeafArrays;
 		TScopedStructBinding<FComplexArrays> ComplexArrays;
-		TestSaveAndLoad(
-			[](FBatchSaver& Batch)
+		Run([](FBatchSaver& Batch)
 			{
 				Batch.Save(FComplexArrays{});
 				Batch.Save(FComplexArrays{{'a', 'b'}, {EFlat1::A}, {{}, {{true}, {2}}}, {{EFlat1::B}, {}} });
@@ -616,8 +626,7 @@ TEST_CASE_NAMED(FPlainPropsUeCoreTest, "System::Core::Serialization::PlainProps:
 	{
 		TScopedStructBinding<FInt> Int;
 		TScopedStructBinding<FUniquePtrs> UniquePtrs;
-		TestSaveAndLoad(
-			[](FBatchSaver& Batch)
+		Run([](FBatchSaver& Batch)
 			{
 				Batch.Save(FUniquePtrs{});
 				Batch.Save(FUniquePtrs{MakeOne(true), MakeOne(FInt{3}), MakeOne(MakeOne(2)), MakeTwo(1.0, 2.0)});
@@ -632,8 +641,7 @@ TEST_CASE_NAMED(FPlainPropsUeCoreTest, "System::Core::Serialization::PlainProps:
 	SECTION("FName")
 	{
 		TScopedStructBinding<FNames> Names;
-		TestSaveAndLoad(
-			[](FBatchSaver& Batch)
+		Run([](FBatchSaver& Batch)
 			{
 				Batch.Save(FNames{ FName("A"), {FName("Y"), FName("A")} });
 			}, 
@@ -643,53 +651,47 @@ TEST_CASE_NAMED(FPlainPropsUeCoreTest, "System::Core::Serialization::PlainProps:
 			});
 	}
 
-	SECTION("Reference")
-	{}
-
 	SECTION("FString")
-	{}
+	{
+		TScopedStructBinding<FStr> Str;
+		Run([](FBatchSaver& Batch)
+			{
+				Batch.Save(FStr{});
+				Batch.Save(FStr{"ABC"});
+				if constexpr (sizeof(TCHAR) > 1)
+				{
+					Batch.Save(FStr{TEXT("\x7FF")});
+					Batch.Save(FStr{TEXT("\x3300")});
+					Batch.Save(FStr{TEXT("\xFE30")});
+					Batch.Save(FStr{TEXT("\xD83D\xDC69")});
+				}
+			}, 
+			[](FBatchLoader& Batch)
+			{
+				CHECK(Batch.Load<FStr>().S.IsEmpty());
+				CHECK(Batch.Load<FStr>().S == "ABC");
+				if constexpr (sizeof(TCHAR) > 1)
+				{
+					CHECK(Batch.Load<FStr>().S == TEXT("\x7FF"));
+					CHECK(Batch.Load<FStr>().S == TEXT("\x3300"));
+					CHECK(Batch.Load<FStr>().S == TEXT("\xFE30"));
+					CHECK(Batch.Load<FStr>().S == TEXT("\xD83D\xDC69"));
+				}
+			});	
+	}
 		
 	SECTION("TSet")
 	{}
 
 	SECTION("NestedContainer")
 	{}
-	
-	SECTION("TSetDelta")
+
+	SECTION("Reference")
 	{}
 
 	//SECTION("LeafOptional")
 	//{}
-	//
-	//SECTION("LeafSmartPtr")
-	//{}
-
-	//SECTION("LeafSetWhole")
-	//{}
-
-	//SECTION("LeafSparseArrayAppends")
-	//{}
-
-	//SECTION("LeafSetOps")
-	//{}
-	//
-	//SECTION("SparseStructArray")
-	//{}
 
-	//SECTION("DenseStructArray")
-	//{}
-	//
-	//SECTION("SubStructArray")
-	//{}
-
-	//SECTION("NestedLeafArray")
-	//{}
-
-	//SECTION("NestedStructArray")
-	//{}
-
-	//SECTION("StructToSubStructMapOps")
-	//{}
 }
 
 } // namespace PlainProps::UE::Test
