PlainProps add delta serialization

#rnx

diff --git a/Engine/Source/Programs/PlainPropsTests/Private/PlainPropsBind.cpp b/Engine/Source/Programs/PlainPropsTests/Private/PlainPropsBind.cpp
--- a/Engine/Source/Programs/PlainPropsTests/Private/PlainPropsBind.cpp
+++ b/Engine/Source/Programs/PlainPropsTests/Private/PlainPropsBind.cpp
@@ -47,6 +47,12 @@ FMemberBindType	FMemberVisitor::PeekType() const
 	return Schema.Members[MemberIdx];
 }
 
+uint32 FMemberVisitor::PeekOffset() const
+{
+	check(HasMore());
+	return Schema.GetOffsets()[MemberIdx];
+}
+
 uint64 FMemberVisitor::GrabMemberOffset()
 {
 	return Schema.GetOffsets()[MemberIdx++];
diff --git a/Engine/Source/Programs/PlainPropsTests/Private/PlainPropsSave.cpp b/Engine/Source/Programs/PlainPropsTests/Private/PlainPropsSave.cpp
--- a/Engine/Source/Programs/PlainPropsTests/Private/PlainPropsSave.cpp
+++ b/Engine/Source/Programs/PlainPropsTests/Private/PlainPropsSave.cpp
@@ -212,23 +212,28 @@ template<class SaverType, typename InnerContextType>
 	return { FMemberType(Member.RangeBindings[0].GetSizeType()), Member.InnermostSchema, ToMemberTypes(Member.InnerTypes) };
 }
 
-static void SaveMember(FMemberBuilder& Out, const uint8* Struct, FMemberId Name, const FSaveContext& Ctx, FLeafMemberBinding Member)
+static const uint8* At(const void* Ptr, SIZE_T Offset)
+{
+	return static_cast<const uint8*>(Ptr) + Offset;
+}
+
+static void SaveMember(FMemberBuilder& Out, const void* Struct, FMemberId Name, const FSaveContext& Ctx, FLeafMemberBinding Member)
 {
 	FUnpackedLeafType Type = { ToLeafType(Member.Leaf.Type), Member.Leaf.Width };
-	Out.AddLeaf(Name, Type, Member.Enum, SaveLeaf(Struct + Member.Offset, Member.Leaf));
+	Out.AddLeaf(Name, Type, Member.Enum, SaveLeaf(At(Struct, Member.Offset), Member.Leaf));
 }
 
-static void SaveMember(FMemberBuilder& Out, const uint8* Struct, FMemberId Name, const FSaveContext& Ctx, FRangeMemberBinding Member)
+static void SaveMember(FMemberBuilder& Out, const void* Struct, FMemberId Name, const FSaveContext& Ctx, FRangeMemberBinding Member)
 {
-	Out.AddRange(Name, { MakeSchema(Member), SaveRange(Struct + Member.Offset, Member, Ctx) });
+	Out.AddRange(Name, { MakeSchema(Member), SaveRange(At(Struct, Member.Offset), Member, Ctx) });
 }
 
-static void SaveMember(FMemberBuilder& Out, const uint8* Struct, FMemberId Name, const FSaveContext& Ctx, FStructMemberBinding Member)
+static void SaveMember(FMemberBuilder& Out, const void* Struct, FMemberId Name, const FSaveContext& Ctx, FStructMemberBinding Member)
 {
-	Out.AddStruct(Name, Member.Id, SaveStruct(Struct + Member.Offset, Member.Id, Ctx));
+	Out.AddStruct(Name, Member.Id, SaveStruct(At(Struct, Member.Offset), Member.Id, Ctx));
 }
 
-TUniquePtr<FBuiltStruct> SaveStruct(const uint8* Struct, FStructSchemaId Id, const FSaveContext& Ctx)
+TUniquePtr<FBuiltStruct> SaveStruct(const void* Struct, FStructSchemaId Id, const FSaveContext& Ctx)
 {
 	const FStructDeclaration& Declaration = Ctx.Declarations.Get(Id);
 
@@ -251,4 +256,221 @@ TUniquePtr<FBuiltStruct> SaveStruct(const uint8* Struct, FStructSchemaId Id, con
 	return Out.BuildAndReset(Declaration, Ctx.Declarations.GetDebug());
 }
 
+////////////////////////////////////////////////////////////////////////////////////////////////
+
+static bool DiffItem(const void* A, const void* B, const FSaveContext& Ctx, FRangeMemberBinding Range);
+static bool DiffItem(const void* A, const void* B, const FSaveContext& Ctx, FStructSchemaId Struct);
+
+static int64 DiffLeaf(const uint8* A, const uint8* B, FUnpackedLeafBindType Leaf)
+{
+	if (Leaf.Type == ELeafBindType::BitfieldBool)
+	{
+		uint32 Mask = 1u << Leaf.BitfieldIdx;
+		return (*A & Mask) - (*B & Mask);
+	}
+
+	switch (Leaf.Width)
+	{
+	case ELeafWidth::B8:	return FMemory::Memcmp(A, B, 1);
+	case ELeafWidth::B16:	return FMemory::Memcmp(A, B, 2);
+	case ELeafWidth::B32:	return FMemory::Memcmp(A, B, 4);
+	case ELeafWidth::B64:	return FMemory::Memcmp(A, B, 8);
+	}
+
+	check(false);
+	return 0;
+}
+
+static bool DiffItemSlice(const uint8* A, const uint8* B, uint64 Num, uint32 Stride, const FSaveContext&, SIZE_T LeafSize)
+{
+	return !!FMemory::Memcmp(A, B, Num * Stride);
+}
+
+template<typename ItemType>
+static bool DiffItemSlice(const uint8* A, const uint8* B, uint64 Num, uint32 Stride, const FSaveContext& Ctx, ItemType&& Member)
+{
+	for (const uint8* EndA = A + Num * Stride; A != EndA; A += Stride, B += Stride)
+	{
+		if (DiffItem(A, B, Ctx, Member))
+		{
+			return true;
+		}
+	}
+
+	return false;
+}
+
+struct FItemRangeReader
+{
+	explicit FItemRangeReader(const void* Range, const IItemRangeBinding& Binding)
+	: Ctx{{Range}}
+	{
+		ReadItems(Binding);
+	}
+
+	FSaveRangeContext Ctx;
+	const uint8* SliceIt = nullptr;
+	uint64 SliceNum = 0;
+
+	void ReadItems(const IItemRangeBinding& Binding)
+	{
+		Binding.ReadItems(Ctx);
+		SliceIt = static_cast<const uint8*>(Ctx.Items.Slice.Data);
+		SliceNum = Ctx.Items.Slice.Num;
+	}
+
+	void RefillItems(const IItemRangeBinding& Binding)
+	{
+		if (SliceNum == 0)
+		{
+			ReadItems(Binding);
+			check(SliceNum > 0);
+		}
+	}
+
+	const uint8* GrabItems(uint64 Num, uint32 Stride)
+	{
+		const uint8* Out = SliceIt;
+		SliceIt += Num * Stride; 
+		SliceNum -= Num;
+		return Out;
+	}
+};
+
+template<class T>
+bool DiffItemRange(const void* RangeA, const void* RangeB, const IItemRangeBinding& Binding, const FSaveContext& OuterCtx, T&& ItemSchema)
+{
+	FItemRangeReader A(RangeA, Binding);
+	FItemRangeReader B(RangeB, Binding);
+	if (A.Ctx.Items.NumTotal != B.Ctx.Items.NumTotal)
+	{
+		return true;
+	}
+
+	if (const uint64 NumTotal = A.Ctx.Items.NumTotal)
+	{
+		check(A.Ctx.Items.Stride == B.Ctx.Items.Stride);
+		const uint32 Stride = A.Ctx.Items.Stride;
+		while (true)
+		{
+			uint64 Num = FMath::Min(A.SliceNum, B.SliceNum);
+			if (DiffItemSlice(A.GrabItems(Num, Stride), B.GrabItems(Num, Stride), Num, Stride, OuterCtx, ItemSchema))
+			{
+				return true;
+			}
+			else if (A.Ctx.Request.NumRead >= NumTotal)
+			{
+				check(A.Ctx.Request.NumRead == NumTotal);	
+				check(B.Ctx.Request.NumRead == NumTotal);	
+				return false;
+			}
+			
+			A.RefillItems(Binding);
+			B.RefillItems(Binding);
+		}
+	}
+
+	return false;
+}
+
+static bool DiffItem(const uint8* A, const uint8* B, const FSaveContext& Ctx, FRangeMemberBinding Member)
+{
+	FRangeBinding Binding = Member.RangeBindings[0];
+	FMemberBindType InnerType = Member.InnerTypes[0];
+
+	if (Binding.IsLeafBinding())
+	{
+		return !!Binding.AsLeafBinding().DiffLeaves(A, B);
+	}
+
+	const IItemRangeBinding& ItemBinding = Binding.AsItemBinding();
+	switch (InnerType.GetKind())
+	{
+	case EMemberKind::Leaf:		return DiffItemRange(A, B, ItemBinding, Ctx, SizeOf(GetArithmeticWidth(InnerType.AsLeaf())));
+	case EMemberKind::Range:	return DiffItemRange(A, B, ItemBinding, Ctx, GetInnerRange(Member));
+	case EMemberKind::Struct:	return DiffItemRange(A, B, ItemBinding, Ctx, static_cast<FStructSchemaId>(Member.InnermostSchema.Get()));
+	}
+
+	check(false);
+	return false;
+}
+
+static bool DiffItem(const uint8* A, const uint8* B, const FSaveContext& Ctx, FStructSchemaId Id)
+{
+	if (ICustomBinding* Custom = Ctx.Customs.FindStruct(Id))
+	{
+		return Custom->DiffStruct(A, B);
+	}
+	
+	bool bOut = false;
+	for (FMemberVisitor It(Ctx.Schemas.GetStruct(Id)); It.HasMore() && !bOut; )
+	{
+		uint32 Offset = It.PeekOffset();
+		const uint8* ItemA = A + Offset;
+		const uint8* ItemB = B + Offset;
+
+		switch (It.PeekKind())
+		{
+		case EMemberKind::Leaf:		bOut = !!DiffLeaf(ItemA, ItemB, It.GrabLeaf().Leaf);	break;
+		case EMemberKind::Range:	bOut = DiffItem(ItemA, ItemB, Ctx, It.GrabRange());		break;
+		case EMemberKind::Struct:	bOut = DiffItem(ItemA, ItemB, Ctx, It.GrabStruct().Id);	break;
+		}
+	}
+
+	return bOut;
+}
+
+////////////////////////////////////////////////////////////////////////////////////////////////
+
+static void SaveMemberDelta(FMemberBuilder& Out, const void* Struct, const void* Default, FMemberId Name, const FSaveContext& Ctx, FLeafMemberBinding Member)
+{
+	if (DiffLeaf(At(Struct, Member.Offset), At(Default, Member.Offset), Member.Leaf))
+	{
+		SaveMember(Out, Struct, Name, Ctx, Member);
+	}
+}
+
+static void SaveMemberDelta(FMemberBuilder& Out, const void* Struct, const void* Default, FMemberId Name, const FSaveContext& Ctx, FRangeMemberBinding Member)
+{
+	const uint8* Range = At(Struct, Member.Offset);
+	if (DiffItem(Range, At(Default, Member.Offset), Ctx, Member))
+	{
+		Out.AddRange(Name, { MakeSchema(Member), SaveRange(Range, Member, Ctx) });
+	}
+}
+
+static void SaveMemberDelta(FMemberBuilder& Out, const void* Struct, const void* Default, FMemberId Name, const FSaveContext& Ctx, FStructMemberBinding Member)
+{
+	if (TUniquePtr<FBuiltStruct> Delta = SaveStructDelta(At(Struct, Member.Offset), At(Default, Member.Offset), Member.Id, Ctx))
+	{
+		Out.AddStruct(Name, Member.Id, MoveTemp(Delta));
+	}
+}
+
+TUniquePtr<FBuiltStruct> SaveStructDelta(const void* Struct, const void* Default, FStructSchemaId Id, const FSaveContext& Ctx)
+{
+	const FStructDeclaration& Declaration = Ctx.Declarations.Get(Id);
+
+	FMemberBuilder Out;
+	if (ICustomBinding* Custom = Ctx.Customs.FindStruct(Id))
+	{
+		if (Custom->DiffStruct(Struct, Default))
+		{
+			Custom->SaveStruct(Out, Struct, Default, Ctx.Declarations.GetDebug());
+		}
+	}
+	else for (FMemberVisitor It(Ctx.Schemas.GetStruct(Id)); It.HasMore(); )
+	{
+		FMemberId Name = Declaration.GetMemberOrder()[It.GetIndex()];
+		switch (It.PeekKind())
+		{
+			case EMemberKind::Leaf:		SaveMemberDelta(Out, Struct, Default, Name, Ctx, It.GrabLeaf());	break;
+			case EMemberKind::Range:	SaveMemberDelta(Out, Struct, Default, Name, Ctx, It.GrabRange());	break;
+			case EMemberKind::Struct:	SaveMemberDelta(Out, Struct, Default, Name, Ctx, It.GrabStruct());	break;
+		}
+	}
+	
+	return !Out.IsEmpty() ? Out.BuildAndReset(Declaration, Ctx.Declarations.GetDebug()) : nullptr;
+}
+
 } // namespace PlainProps
\ No newline at end of file
diff --git a/Engine/Source/Programs/PlainPropsTests/Public/PlainPropsBind.h b/Engine/Source/Programs/PlainPropsTests/Public/PlainPropsBind.h
--- a/Engine/Source/Programs/PlainPropsTests/Public/PlainPropsBind.h
+++ b/Engine/Source/Programs/PlainPropsTests/Public/PlainPropsBind.h
@@ -157,7 +157,7 @@ struct FUnpackedLeafBindType
 		}
 	}
 
-//	constexpr bool operator==(FUnpackedLeafBindType O) { return Type == O.Type && Width == O.Width; }
+//	constexpr bool operator==(FUnpackedLeafBindType O) const { return Type == O.Type && Width == O.Width; }
 	FMemberBindType Pack() const
 	{ 
 		return FMemberBindType(Type == ELeafBindType::BitfieldBool ? FLeafBindType(BitfieldIdx) : FLeafBindType(Type, Width));
@@ -175,7 +175,7 @@ struct FLeafMemberBinding
 {
 	FUnpackedLeafBindType	Leaf;
 	FOptionalEnumSchemaId	Enum;
-	uint64					Offset;
+	SIZE_T					Offset;
 };
 
 struct FRangeMemberBinding
@@ -184,14 +184,14 @@ struct FRangeMemberBinding
 	const FRangeBinding*	RangeBindings;
 	uint32					NumRanges; // At least 1, >1 for nested ranges
 	FOptionalSchemaId		InnermostSchema;
-	uint64					Offset;
+	SIZE_T					Offset;
 };
 
 struct FStructMemberBinding
 {
-	FStructType			Type;
-	FStructSchemaId		Id;
-	uint64				Offset;
+	FStructType				Type;
+	FStructSchemaId			Id;
+	SIZE_T					Offset;
 };
 
 // Iterates over member bindings
@@ -205,6 +205,7 @@ public:
 	
 	EMemberKind					PeekKind() const;		// @pre HasMore()
 	FMemberBindType				PeekType() const;		// @pre HasMore()
+	uint32						PeekOffset() const;		// @pre HasMore()
 
 	FLeafMemberBinding			GrabLeaf();				// @pre PeekKind() == EMemberKind::Leaf
 	FRangeMemberBinding			GrabRange();			// @pre PeekKind() == EMemberKind::Range
@@ -236,8 +237,9 @@ class ICustomBinding
 {
 public:
 	virtual ~ICustomBinding() {}
-	virtual void				SaveStruct(FMemberBuilder& Dst, const void* Src, void* UserData, const FDebugIds& Debug) = 0;
+	virtual void				SaveStruct(FMemberBuilder& Dst, const void* Src, const void* Default, const FDebugIds& Debug) = 0;
 	virtual void				LoadStruct(void* Dst, FStructView Src, ECustomLoadMethod Method, const FLoadBatch& Batch) const = 0;
+	virtual bool				DiffStruct(const void* StructA, const void* StructB) const = 0;
 };
 
 class FCustomBindings
@@ -461,8 +463,9 @@ public:
 class alignas(16) ILeafRangeBinding
 {
 public:
-	virtual void SaveLeaves(const void* Range, FLeafRangeAllocator& Out) const = 0;
-	virtual void LoadLeaves(void* Range, FLeafRangeLoadView Leaves) const = 0;
+	virtual void	SaveLeaves(const void* Range, FLeafRangeAllocator& Out) const = 0;
+	virtual void	LoadLeaves(void* Range, FLeafRangeLoadView Leaves) const = 0;
+	virtual int64	DiffLeaves(const void* RangeA, const void* RangeB) const = 0;
 };
 
 ////////////////////////////////////////////////////////////////////////////////////////////////
diff --git a/Engine/Source/Programs/PlainPropsTests/Public/PlainPropsBuild.h b/Engine/Source/Programs/PlainPropsTests/Public/PlainPropsBuild.h
--- a/Engine/Source/Programs/PlainPropsTests/Public/PlainPropsBuild.h
+++ b/Engine/Source/Programs/PlainPropsTests/Public/PlainPropsBuild.h
@@ -176,6 +176,8 @@ public:
 
 	PP_API [[nodiscard]] TUniquePtr<FBuiltStruct> BuildAndReset(const FStructDeclaration& Declared, const FDebugIds& Debug);
 
+	bool IsEmpty() const { return Members.IsEmpty(); }
+
 private:
 	TArray<FBuiltMember>	Members;
 	
diff --git a/Engine/Source/Programs/PlainPropsTests/Public/PlainPropsSave.h b/Engine/Source/Programs/PlainPropsTests/Public/PlainPropsSave.h
--- a/Engine/Source/Programs/PlainPropsTests/Public/PlainPropsSave.h
+++ b/Engine/Source/Programs/PlainPropsTests/Public/PlainPropsSave.h
@@ -20,7 +20,8 @@ struct FSaveContext
 	FCustomBindings&			Customs;
 };
 
-[[nodiscard]] TUniquePtr<FBuiltStruct> SaveStruct(const uint8* Struct, FStructSchemaId Id, const FSaveContext& Context);
+[[nodiscard]] TUniquePtr<FBuiltStruct> SaveStruct(const void* Struct, FStructSchemaId Id, const FSaveContext& Context);
+[[nodiscard]] TUniquePtr<FBuiltStruct> SaveStructDelta(const void* Struct, const void* Default, FStructSchemaId Id, const FSaveContext& Context);
 
 } // namespace PlainProps
 
diff --git a/Engine/Source/Programs/PlainPropsTests/Public/PlainPropsUeCoreBindings.h b/Engine/Source/Programs/PlainPropsTests/Public/PlainPropsUeCoreBindings.h
--- a/Engine/Source/Programs/PlainPropsTests/Public/PlainPropsUeCoreBindings.h
+++ b/Engine/Source/Programs/PlainPropsTests/Public/PlainPropsUeCoreBindings.h
@@ -186,6 +186,22 @@ struct FStringBinding : public ILeafRangeBinding
 			*DstEnd = '\0';
 		}
 	}
+
+	virtual int64 DiffLeaves(const void* RangeA, const void* RangeB) const override
+	{
+		const FString& A = *static_cast<const FString*>(RangeA);
+		const FString& B = *static_cast<const FString*>(RangeB);
+		int32 ALen = A.Len();
+		int32 BLen = B.Len();
+
+		if (int32 LenDiff = ALen - BLen)
+		{
+			return LenDiff;
+		}
+
+		// Case-sensitive comparison
+		return ALen ? FMemory::Memcmp(A.GetCharArray().GetData(), B.GetCharArray().GetData(), ALen * sizeof(TCHAR)) : 0;
+	}
 };
 
 //////////////////////////////////////////////////////////////////////////
@@ -366,11 +382,11 @@ struct TSetDeltaBinding : public ICustomBinding
 	
 	static TConstArrayView<FMemberId> GetMemberIds() { return FSetOps::Get<Ids>().All; }
 
-	virtual void SaveStruct(FMemberBuilder& Dst, const void* Src, void* UserData, const FDebugIds& Debug) override;
+	virtual void SaveStruct(FMemberBuilder& Dst, const void* Src, const void* Default, const FDebugIds& Debug) override;
 
 	virtual void LoadStruct(void* Dst, FStructView Src, ECustomLoadMethod Method, const FLoadBatch& Batch) const override
 	{
-		Type& Out = *reinterpret_cast<Type*>(Dst);
+		Type& Out = *static_cast<Type*>(Dst);
 		FMemberReader Members(Src);
 
 		if (Method == ECustomLoadMethod::Construct)
@@ -413,6 +429,26 @@ struct TSetDeltaBinding : public ICustomBinding
 		check(!Members.HasMore());
 	}
 
+	virtual bool DiffStruct(const void* StructA, const void* StructB) const override
+	{
+		const Type& A = *static_cast<const Type*>(StructA);
+		const Type& B = *static_cast<const Type*>(StructB);
+		if (A.Num() != B.Num())
+		{
+			return false;
+		}
+
+		for (const T& AKey : A)
+		{
+			if (!B.Contains(AKey))
+			{
+				return false;
+			}
+		}
+
+		return true;
+	}
+
 	void AddItems(Type& Out, FRangeView Items, const FLoadBatch& Batch)
 	{
 		check(!Items.IsEmpty());
diff --git a/Engine/Source/Programs/PlainPropsTests/Tests/PlainPropsUeCoreTests.cpp b/Engine/Source/Programs/PlainPropsTests/Tests/PlainPropsUeCoreTests.cpp
--- a/Engine/Source/Programs/PlainPropsTests/Tests/PlainPropsUeCoreTests.cpp
+++ b/Engine/Source/Programs/PlainPropsTests/Tests/PlainPropsUeCoreTests.cpp
@@ -137,7 +137,7 @@ struct FTestCustomBinding : public ICustomBinding
 
 struct FNameBinding : public FTestCustomBinding
 {
-	virtual void SaveStruct(FMemberBuilder& Dst, const void* Src, void*, const FDebugIds& Debug) override
+	virtual void SaveStruct(FMemberBuilder& Dst, const void* Src, const void*, const FDebugIds& Debug) override
 	{
 		FSetElementId Idx = Names.Add(*static_cast<const FName*>(Src));
 		Dst.Add(Declaration.Idx, Idx.AsInteger());
@@ -149,6 +149,13 @@ struct FNameBinding : public FTestCustomBinding
 		*static_cast<FName*>(Dst) = Names.Get(Idx);
 	}
 
+	virtual bool DiffStruct(const void* StructA, const void* StructB) const override
+	{
+		FName A = *static_cast<const FName*>(StructA);
+		FName B = *static_cast<const FName*>(StructB);
+		return A.IsEqual(B, ENameCase::CaseSensitive);
+	}
+
 	virtual FStructSchemaId GetId() const override { return Declaration.Id; }
 
 	FNameDeclaration	Declaration;
@@ -166,6 +173,9 @@ public:
 
 	template<class T>
 	void						Save(T&& Object);
+	template<class T>
+	bool						SaveDelta(const T& Object, const T& Default);
+
 	TArray64<uint8>				Write() const;
 
 private:
@@ -185,7 +195,19 @@ template<class T>
 void FBatchSaver::Save(T&& Object) 
 {
 	FStructSchemaId Id = IndexNativeStruct<T, FIds>();
-	SavedObjects.Emplace(Id, SaveStruct(reinterpret_cast<const uint8*>(&Object), Id, {GTypes, GSchemas, Customs}));
+	SavedObjects.Emplace(Id, SaveStruct(&Object, Id, {GTypes, GSchemas, Customs}));
+}
+
+template<class T>
+bool FBatchSaver::SaveDelta(const T& Object, const T& Default) 
+{
+	FStructSchemaId Id = IndexNativeStruct<T, FIds>();
+	if (TUniquePtr<FBuiltStruct> Delta = SaveStructDelta(&Object, &Default, Id, {GTypes, GSchemas, Customs}))
+	{
+		SavedObjects.Emplace(Id, MoveTemp(Delta));
+		return true;
+	}
+	return false;
 }
 
 template<typename ArrayType>
@@ -318,9 +340,15 @@ public:
 	T Load()
 	{
 		T Out;
+		LoadInto(Out);
+		return MoveTemp(Out);
+	}
+
+	template<class T>
+	void LoadInto(T& Out)
+	{
 		FStructView In = Objects[LoadIdx++];
 		LoadStruct(reinterpret_cast<uint8*>(&Out), In, *Plans);
-		return MoveTemp(Out);
 	}
 
 private:
@@ -480,6 +508,19 @@ TArray<TUniquePtr<T>> MakeTwo(T&& A, T&& B)
 	return Out;
 }
 
+//////////////////////////////////////////////////////////////////////////
+
+struct FDelta
+{
+	bool		A = true;
+	float		B = 1.0;
+	FInt		C = { 2 };
+	TArray<int> D;
+	FString		E = "!";
+};
+PP_REFLECT_STRUCT(PlainProps::UE::Test, FDelta, void, A, B, C, D, E);
+static bool operator==(const FDelta& A, const FDelta& B) { return A.A == B.A && A.B == B.B && A.C == B.C && A.D == B.D && A.E == B.E; }
+
 //////////////////////////////////////////////////////////////////////////
 //
 //struct FObject
@@ -501,7 +542,7 @@ TArray<TUniquePtr<T>> MakeTwo(T&& A, T&& B)
 //		return MakeArrayView(Ids, 2);
 //	}
 //
-//	virtual void SaveStruct(FMemberBuilder& Dst, const void* Src, void*, const FDebugIds&) const override
+//	virtual void SaveStruct(FMemberBuilder& Dst, const void* Src, const void*, const FDebugIds&) const override
 //	{
 //		if (const FObject* Object = reinterpret_cast<const FObject*>(Src))
 //		{
@@ -683,15 +724,40 @@ TEST_CASE_NAMED(FPlainPropsUeCoreTest, "System::Core::Serialization::PlainProps:
 	SECTION("TSet")
 	{}
 
-	SECTION("NestedContainer")
-	{}
+	SECTION("Delta")
+	{
+		TScopedStructBinding<FInt> Int;
+		TScopedStructBinding<FDelta> Delta;
+		Run([](FBatchSaver& Batch)
+			{
+				FDelta Zero = {false, 0, {}, {}, {}};
+				CHECK(!Batch.SaveDelta(FInt{123},FInt{123}));
+				CHECK(!Batch.SaveDelta(FDelta{},FDelta{}));
+				CHECK(!Batch.SaveDelta(Zero, Zero));
+					
+				Batch.SaveDelta({}, Zero);
+				Batch.SaveDelta(Zero, {});
+				Batch.SaveDelta(FDelta{.B = 123}, {});
+				Batch.SaveDelta(FDelta{.C = {321}}, {});
+				Batch.SaveDelta(FDelta{.D = {0}}, {});
+				Batch.SaveDelta(FDelta{.E = "!!"}, {});
+			}, 
+			[](FBatchLoader& Batch)
+			{
+				FDelta Zero = {false, 0, {}, {}, {}};
+				FDelta DefaultOnZero = Zero;
+				Batch.LoadInto(DefaultOnZero);
+				CHECK(DefaultOnZero == FDelta{});
+				CHECK(Batch.Load<FDelta>() == Zero);
+				CHECK(Batch.Load<FDelta>() == FDelta{.B = 123});
+				CHECK(Batch.Load<FDelta>() == FDelta{.C = {321}});
+				CHECK(Batch.Load<FDelta>() == FDelta{.D = {0}});
+				CHECK(Batch.Load<FDelta>() == FDelta{.E = "!!"});
+			});
+	}
 
 	SECTION("Reference")
 	{}
-
-	//SECTION("LeafOptional")
-	//{}
-
 }
 
 } // namespace PlainProps::UE::Test
