PlainProps: Fix crashes when launching PlainPropsTests without a test filter.

#rnx

#Synced-CL 52129858

diff --git a/Engine/Source/Programs/PlainPropsTests/PlainPropsTests.Build.cs b/Engine/Source/Programs/PlainPropsTests/PlainPropsTests.Build.cs
--- a/Engine/Source/Programs/PlainPropsTests/PlainPropsTests.Build.cs
+++ b/Engine/Source/Programs/PlainPropsTests/PlainPropsTests.Build.cs
@@ -30,6 +30,8 @@ public class PlainPropsTests : TestModuleRules
 	{
 		PrivateDependencyModuleNames.AddRange(
 			new string[] {
+				"Core",					// For Misc/ConfigCacheIni.h test setup
+				"LowLevelTestsRunner",	// For TestCommon/Initialization.h test setup
 				"PlainProps"
 			});
 
@@ -38,4 +40,4 @@ public class PlainPropsTests : TestModuleRules
 			PrivateDependencyModuleNames.AddRange(new string[] { "DesktopPlatform" });
 		}
 	}
-}
\ No newline at end of file
+}
diff --git a/Engine/Source/Programs/PlainPropsTests/TestGroupEvents.cpp b/Engine/Source/Programs/PlainPropsTests/TestGroupEvents.cpp
new file mode 100644
--- /dev/null
+++ b/Engine/Source/Programs/PlainPropsTests/TestGroupEvents.cpp
@@ -0,0 +1,18 @@
+// Copyright Epic Games, Inc. All Rights Reserved.
+#include "TestCommon/Initialization.h"
+#include "Misc/ConfigCacheIni.h"
+
+#include <catch2/catch_test_macros.hpp>
+
+GROUP_BEFORE_GLOBAL(Catch::DefaultGroup)
+{
+	GConfig = new FConfigCacheIni(EConfigCacheType::Temporary); // for URLRequestFilterTest
+	InitAll(/* bAllowLogging*/ true, /* bMultithreaded*/ true);
+}
+
+GROUP_AFTER_GLOBAL(Catch::DefaultGroup)
+{
+	CleanupAll();
+	delete GConfig;
+	GConfig = nullptr;
+}
